Copied fixes for MDL, AxF, GLTF and the Variants Manager: 5224459, 5196836, 5196494, 5060622, 4928742 from //Tasks/UE4/Private-Datasmith-Deltagen/

diff --git a/Engine/Plugins/Editor/GLTFImporter/Source/GLTFImporter/Private/GLTF/GLTFReader.cpp b/Engine/Plugins/Editor/GLTFImporter/Source/GLTFImporter/Private/GLTF/GLTFReader.cpp
--- a/Engine/Plugins/Editor/GLTFImporter/Source/GLTFImporter/Private/GLTF/GLTFReader.cpp
+++ b/Engine/Plugins/Editor/GLTFImporter/Source/GLTFImporter/Private/GLTF/GLTFReader.cpp
@@ -318,7 +318,12 @@ namespace GLTF
 	{
 		const uint32 CameraIndex = Asset->Cameras.Num();
 		const FNode* Found       = Asset->Nodes.FindByPredicate([CameraIndex](const FNode& Node) { return CameraIndex == Node.CameraIndex; });
-		check(Found);
+		FString Name = GetString(Object, TEXT("name"));
+		if (!Found)
+		{
+			Messages.Emplace(EMessageSeverity::Warning, FString::Printf(TEXT("No camera node found for camera %d('%s')"), CameraIndex, *Name));
+			return;
+		}
 
 		FCamera& Camera = Asset->Cameras.Emplace_GetRef(*Found);
 		Camera.Name     = GetString(Object, TEXT("name"));
diff --git a/Engine/Plugins/Enterprise/VariantManagerContent/Source/VariantManagerContent/Private/PropertyValue.cpp b/Engine/Plugins/Enterprise/VariantManagerContent/Source/VariantManagerContent/Private/PropertyValue.cpp
--- a/Engine/Plugins/Enterprise/VariantManagerContent/Source/VariantManagerContent/Private/PropertyValue.cpp
+++ b/Engine/Plugins/Enterprise/VariantManagerContent/Source/VariantManagerContent/Private/PropertyValue.cpp
@@ -876,6 +876,12 @@ void UPropertyValue::ApplyViaFunctionSetter(UObject* TargetObject)
 		}
 	}
 
+	// HACK: Restore Visibility properties to operating recursively. Temporary until 4.23
+	if (PropertySetter->GetName() == TEXT("SetVisibility") && PropertySetter->ParmsSize == 2 && Parms)
+	{
+		Parms[1] = true;
+	}
+
 	// Only actually call the function if we managed to pack our recorded bytes in the params. Else we will
 	// just reset the object to defaults
 	if (bAppliedRecordedData)
diff --git a/Engine/Plugins/Enterprise/VariantManagerContent/Source/VariantManagerContent/Private/VariantObjectBinding.cpp b/Engine/Plugins/Enterprise/VariantManagerContent/Source/VariantManagerContent/Private/VariantObjectBinding.cpp
--- a/Engine/Plugins/Enterprise/VariantManagerContent/Source/VariantManagerContent/Private/VariantObjectBinding.cpp
+++ b/Engine/Plugins/Enterprise/VariantManagerContent/Source/VariantManagerContent/Private/VariantObjectBinding.cpp
@@ -227,7 +227,7 @@ void UVariantObjectBinding::ExecuteTargetFunction(FName FunctionName)
 	UWorld* World = DirectorInstance->GetWorld();
 	if (World->WorldType == EWorldType::Editor && !Func->HasMetaData(NAME_CallInEditor))
 	{
-		UE_LOG(LogVariantContent, Warning, TEXT("Cannot call function '%s' as it doesn't have the CallInEditor option checked! It is also recommended to only ever call other CallInEditor functions from within it."), *FunctionName.ToString());
+		UE_LOG(LogVariantContent, Warning, TEXT("Cannot call function '%s' as it doesn't have the CallInEditor option checked! Also note that calling this from the editor may have irreversible effects on the level."), *FunctionName.ToString());
 		return;
 	}
 #endif
@@ -291,7 +291,7 @@ void UVariantObjectBinding::ExecuteAllTargetFunctions()
 		UWorld* World = DirectorInstance->GetWorld();
 		if (World->WorldType == EWorldType::Editor && !Func->HasMetaData(NAME_CallInEditor))
 		{
-			UE_LOG(LogVariantContent, Warning, TEXT("Cannot call function '%s' as it doesn't have the CallInEditor option checked! It is also recommended to only ever call other CallInEditor functions from within it."), *Func->GetName());
+			UE_LOG(LogVariantContent, Warning, TEXT("Cannot call function '%s' as it doesn't have the CallInEditor option checked! Also note that calling this from the editor may have irreversible effects on the level."), *Func->GetName());
 			continue;
 		}
 #endif
