Copied cl 5562264

Patch so that Location/Rotation properties work in play mode;
Fix typo when recovering RelativeRotation properties from older project versions;

#rnx

diff --git a/Engine/Plugins/Enterprise/VariantManagerContent/Source/VariantManagerContent/Private/PropertyValue.cpp b/Engine/Plugins/Enterprise/VariantManagerContent/Source/VariantManagerContent/Private/PropertyValue.cpp
--- a/Engine/Plugins/Enterprise/VariantManagerContent/Source/VariantManagerContent/Private/PropertyValue.cpp
+++ b/Engine/Plugins/Enterprise/VariantManagerContent/Source/VariantManagerContent/Private/PropertyValue.cpp
@@ -33,6 +33,11 @@ void UPropertyValue::Init(const TArray<FCapturedPropSegment>& InCapturedPropSegm
 	PropertySetterName = InPropertySetterName;
 	PropCategory = InCategory;
 
+	if (PropertySetterName == TEXT("SetRelativeLocation") || PropertySetterName == TEXT("SetRelativeLocation"))
+	{
+		PropertySetterName = FName(*(TEXT("K2_") + PropertySetterName.ToString()));
+	}
+
 	ClearLastResolve();
 	ValueBytes.SetNumUninitialized(GetValueSizeInBytes());
 	TempObjPtr.Reset();
@@ -110,6 +115,11 @@ void UPropertyValue::Serialize(FArchive& Ar)
 	}
 	else if (Ar.IsLoading())
 	{
+		if (PropertySetterName == TEXT("SetRelativeLocation") || PropertySetterName == TEXT("SetRelativeLocation"))
+		{
+			PropertySetterName = FName(*(TEXT("K2_") + PropertySetterName.ToString()));
+		}
+
 		Ar << TempObjPtr;
 
 		// Before this version, properties were stored an array of UProperty*. Convert them to
@@ -1134,10 +1144,10 @@ void UPropertyValueTransform::Serialize(FArchive& Ar)
 	switch (PropCategory)
 	{
 	case EPropertyValueCategory::RelativeLocation:
-		PropertySetterName = FName(TEXT("SetRelativeLocation"));
+		PropertySetterName = FName(TEXT("K2_SetRelativeLocation"));
 		break;
 	case EPropertyValueCategory::RelativeRotation:
-		PropertySetterName = FName(TEXT("SetRelativeLocation"));
+		PropertySetterName = FName(TEXT("K2_SetRelativeRotation"));
 		break;
 	case EPropertyValueCategory::RelativeScale3D:
 		PropertySetterName = FName(TEXT("SetRelativeScale3D"));
