Fixed copy paste error in property name comparison.

#rnx

diff --git a/Engine/Plugins/Enterprise/VariantManagerContent/Source/VariantManagerContent/Private/PropertyValue.cpp b/Engine/Plugins/Enterprise/VariantManagerContent/Source/VariantManagerContent/Private/PropertyValue.cpp
--- a/Engine/Plugins/Enterprise/VariantManagerContent/Source/VariantManagerContent/Private/PropertyValue.cpp
+++ b/Engine/Plugins/Enterprise/VariantManagerContent/Source/VariantManagerContent/Private/PropertyValue.cpp
@@ -33,7 +33,7 @@ void UPropertyValue::Init(const TArray<FCapturedPropSegment>& InCapturedPropSegm
 	PropertySetterName = InPropertySetterName;
 	PropCategory = InCategory;
 
-	if (PropertySetterName == TEXT("SetRelativeLocation") || PropertySetterName == TEXT("SetRelativeLocation"))
+	if (PropertySetterName == TEXT("SetRelativeLocation") || PropertySetterName == TEXT("SetRelativeRotation"))
 	{
 		PropertySetterName = FName(*(TEXT("K2_") + PropertySetterName.ToString()));
 	}
@@ -115,7 +115,7 @@ void UPropertyValue::Serialize(FArchive& Ar)
 	}
 	else if (Ar.IsLoading())
 	{
-		if (PropertySetterName == TEXT("SetRelativeLocation") || PropertySetterName == TEXT("SetRelativeLocation"))
+		if (PropertySetterName == TEXT("SetRelativeLocation") || PropertySetterName == TEXT("SetRelativeRotation"))
 		{
 			PropertySetterName = FName(*(TEXT("K2_") + PropertySetterName.ToString()));
 		}
