Copied 4401443 JeanMichel.Dignard
VRED, AxF/MDL, VariantManager integration

diff --git a/Engine/Plugins/Enterprise/DatasmithContent/Shaders/Private/FlakesAccess.usf b/Engine/Plugins/Enterprise/DatasmithContent/Shaders/Private/FlakesAccess.usf
deleted file mode 100644
diff --git a/Engine/Plugins/Enterprise/DatasmithContent/Shaders/Private/FlakesAccess.ush b/Engine/Plugins/Enterprise/DatasmithContent/Shaders/Private/FlakesAccess.ush
new file mode 100644
--- /dev/null
+++ b/Engine/Plugins/Enterprise/DatasmithContent/Shaders/Private/FlakesAccess.ush
@@ -0,0 +1,131 @@
+
+/**
+Flakes texture access based on light angular direction, uses a pseduo(i.e. in atlas format) 2D texture array
+to access flakes slices and finally applies bilinear interpolation.
+
+@param ThetaAngles - normalized(by Pi/2) thetaF and thetaI angles
+@param FlakesTexArray - texture 2d array in an atlas format
+@param FlakesTexCount - number of textures in the array
+@param FlakesTexValues - x = texture size of an entry in the 2d array, y = ThetaSliceLUT texel size
+@param TexCoord - UV coordinates
+@param Tiling - UV tiling factor
+@param NumThetaF - number of ThetaF bins
+@param NumThetaI- number of thetaI bins
+@param MaxThetaI - max thetaI bin that is used
+@param ThetaSliceLUT - 1d texture with the index into flake texture array for a given thetaI bin
+@return Bilinear interpolated value of the flakes.
+*/
+
+#define PI2 1.5707963267948966
+#define USE_GRADIENTS 1
+
+struct TextureArraySample
+{ 
+    float4 Sample(float2 UV, float Index)
+    {
+      float SideCount = sqrt(FlakesTexCount); // per side texture count
+      float InvSideCount = 1.0 / SideCount; 
+      float TexPadding = 8.0; // pixels padding count
+      float TexSize = FlakesTexValues.x;
+      float PaddedTexSize = TexSize + TexPadding * 2.0;
+      float AtlasTexelSize = 1.0 / (SideCount * PaddedTexSize);
+
+      float2 TexOffset = Index;
+      TexOffset.x = fmod(TexOffset.x, SideCount);
+      TexOffset.y = floor(TexOffset.y * InvSideCount);
+      TexOffset *= InvSideCount;
+      TexOffset += TexPadding * AtlasTexelSize;
+
+      float2 SampleTexCoord = frac(UV * Tiling);
+      SampleTexCoord *= InvSideCount * (TexSize / PaddedTexSize);
+      SampleTexCoord += TexOffset;
+
+#if USE_GRADIENTS == 1
+      float2 DdxUV = DDX(SampleTexCoord);
+      float2 DdyUV = DDY(SampleTexCoord);
+      return Texture2DSampleGrad(FlakesTexArray, FlakesTexArraySampler, SampleTexCoord, DdxUV, DdyUV);
+#else
+      return Texture2DSample(FlakesTexArray, FlakesTexArraySampler, SampleTexCoord);
+#endif
+    }  
+}; 
+
+struct ThetaSliceLUTSample
+{ 
+    float Sample(float BinThetaI)
+    {
+      float TexelSize = FlakesTexValues.y;
+      float Padding = 4.0;
+      float Index = BinThetaI * Padding * TexelSize + TexelSize;
+
+      float Value = Texture2DSampleLevel(ThetaSliceLUT, ThetaSliceLUTSampler, float2(Index, 0.0), 0.0).x;
+      Value = floor(Value * 255.0 + 0.5);
+      return Value;
+    }
+};
+TextureArraySample TextureArray;
+ThetaSliceLUTSample ThetaFiLUT;
+
+//ThetaF sampling defines the angular sampling, i.e. angular flake lifetime
+float ThetaF = clamp(ThetaAngles.x, 0.0, 1.0 / PI2) * NumThetaF + 0.5;
+float ThetaI = clamp(ThetaAngles.y, 0.0, 1.0 / PI2) * NumThetaF + 0.5;
+
+//bilinear interpolation indices
+float LowThetaF = floor(ThetaF);
+float HighThetaF = LowThetaF + 1.0;
+float LowThetaI = floor(ThetaI);
+float HighThetaI = LowThetaI + 1.0;
+
+float2 UV = TexCoord;
+if (NumThetaI < NumThetaF)
+{
+    if ((LowThetaI % 2) == 1 || (HighThetaI % 2) == 1) 
+    { 
+      //tweak for extra randomness, in the case of almost planar geometries with repeating patches
+      UV.xy = UV.yx;
+    }
+
+    //map to the original sampling
+    float InvNumThetaF = 1.0 / NumThetaF;
+    LowThetaI = int(floor(float(LowThetaI) * NumThetaI * InvNumThetaF));
+    HighThetaI = int(floor(float(HighThetaI) * NumThetaI * InvNumThetaF));
+}
+
+half3 ColorBottomLeft = 0.0;
+half3 ColorTopLeft = 0.0;
+half3 ColorBottomRight = 0.0;
+half3 ColorTopRight  = 0.0;
+//access flakes texture and stay in the correct slices(no slip over)
+if (LowThetaI < MaxThetaI)
+{
+    float LowThetaIndex = ThetaFiLUT.Sample(LowThetaI) + LowThetaF;
+    if (LowThetaIndex < ThetaFiLUT.Sample(LowThetaI + 1.0))
+    {
+        ColorBottomLeft = TextureArray.Sample(UV, LowThetaIndex).xyz;
+    }
+    float HighThetaIndex = ThetaFiLUT.Sample(LowThetaI) + HighThetaF;
+    if (HighThetaIndex < ThetaFiLUT.Sample(LowThetaI + 1))
+    {
+        ColorBottomRight = TextureArray.Sample(UV, HighThetaIndex).xyz;
+    }
+}
+
+if (HighThetaI < MaxThetaI)
+{
+    float LowThetaIndex = ThetaFiLUT.Sample(HighThetaI) + LowThetaF;
+    if (LowThetaIndex < ThetaFiLUT.Sample(HighThetaI + 1))
+    {
+        ColorTopLeft = TextureArray.Sample(UV, LowThetaIndex).xyz;
+    }
+    float HighThetaIndex = ThetaFiLUT.Sample(HighThetaI) + HighThetaF;
+    if (HighThetaIndex < ThetaFiLUT.Sample(HighThetaI + 1))
+    {
+        ColorTopRight = TextureArray.Sample(UV, HighThetaIndex).xyz;
+    }
+}
+//bilinear interpolation
+float DiffThetaF = ThetaF - float(LowThetaF);
+float DiffThetaI = ThetaI - float(LowThetaI);
+half3 Bottom = (1.0 - DiffThetaF) * ColorBottomLeft + DiffThetaF * ColorBottomRight;
+half3 Top = (1.0 - DiffThetaF) * ColorTopLeft + DiffThetaF * ColorTopRight;    
+return (1.0 - DiffThetaI) * Bottom + DiffThetaI * Top;
diff --git a/Engine/Plugins/Enterprise/DatasmithContent/Shaders/Private/SpiralBlurWithIOR.ush b/Engine/Plugins/Enterprise/DatasmithContent/Shaders/Private/SpiralBlurWithIOR.ush
new file mode 100644
--- /dev/null
+++ b/Engine/Plugins/Enterprise/DatasmithContent/Shaders/Private/SpiralBlurWithIOR.ush
@@ -0,0 +1,118 @@
+/*
+Spiral blur function with IOR used to emulate BTDF of the material.
+
+ @param RandomTex - Random texture input for spiral blur
+ @param SurfaceScatter - Scattering by depth behind the object (0.0 to 1.0)
+ @param VolumeScatter - Scattering by object's thickness (0.0 to 1.0)
+ @param IOR - Index of refraction (defaults to 1.0 for no refraction)
+ @param Normal - Normal texture input
+ @param ObjectThickness - Estimated or calculated object's thickness 
+ @param PixelDepth - Pixel depth input (TODO could be sampled in code)
+
+ @param NormalAlignFactor - Factor for Normal to Screen alignment, controls refraction and
+                            keeps it within screen space (0.0 to positive infinity with 0.0
+                            meaning no correction)
+ @param MagicRefractionNumber - Directly controls the amount of refraction offset when sampling
+                                the scene (1.0 means no correction)
+
+*/// --------------------------------------------------------------------------
+
+const int RadialSteps = 48;
+const int AngularSteps = 1;
+
+const float MagicBlurNumber = 100.0;
+
+// Thickness added to account for normal texture features when
+// calculating refraction light bending.
+const float RefractionThicknessCorrection = 50.0;
+
+float UVDistance = 0.2;
+float AngularOffset = 0.618;
+
+float SizeOnScreenFactor = View.ViewToClip[0][0] / PixelDepth;
+
+float2 ScreenUV = MaterialFloat2(ScreenAlignedPosition(Parameters.ScreenPosition).xy);
+float2 BaseUV = ScreenUV;
+
+// Calculate PixelNormal
+float3 PixelNormal = mul(mul(Normal, Parameters.TangentToWorld), (MaterialFloat3x3)ResolvedView.TranslatedWorldToView);
+
+// Fetch SceneDepth
+float SceneDepth = CalcSceneDepth(BaseUV);
+
+// Straighten PixelNormal to keep refraction closer to screen space
+float3 DistortedNormal = PixelNormal;
+DistortedNormal.z += sign(DistortedNormal.z) * NormalAlignFactor;
+DistortedNormal = normalize(DistortedNormal);
+
+// Adjust BaseUV for refraction (based on IOR, PixelNormal, and ObjectThickness)
+float2 NormalUV = DistortedNormal.xy;
+NormalUV.y = -NormalUV.y;
+float2 RefractionOffset = (ObjectThickness + RefractionThicknessCorrection) * NormalUV * (1.0 / IOR - 1.0);
+
+// Establish max RefractionOffset and scale by SizeOnScreenFactor
+RefractionOffset *= MagicRefractionNumber * SizeOnScreenFactor;
+
+BaseUV += RefractionOffset;
+
+float2 NewUV = BaseUV;
+
+float RadialStepSize = UVDistance / RadialSteps;
+float2 CurOffset = { 0.0, 0.0 };
+
+// Recalculate SceneDepth with NewUV
+SceneDepth = CalcSceneDepth(NewUV);
+
+// Screen aligned TemporalAA jitter
+float2 ScenePixels = View.BufferSizeAndInvSize.xy * ScreenUV;
+ScenePixels += View.TemporalAAParams.r;
+float2 RandomSamp = ((uint)(ScenePixels.x) + 2 * (uint)(ScenePixels.y)) % 5;
+RandomSamp += Texture2DSample(RandomTex, RandomTexSampler, View.BufferSizeAndInvSize.xy * ScreenUV / 64);
+RandomSamp /= 6;
+RandomSamp -= 0.5;
+
+float TempAARotation = RandomSamp.x;
+float TempAADistance = RadialStepSize * RandomSamp.x;
+
+
+BaseUV += View.BufferSizeAndInvSize.zw * RandomSamp;
+
+// Fetch RandomSceneDepth (costly but could improve blur on SceneDepth sharp edges)
+// float RandomSceneDepth = CalcSceneDepth(NewUV + RandomSamp * (UVDistance / RadialSteps));
+
+// Calculate BlurAmount
+float VolumeBlurAmount = 1.0 - exp(-ObjectThickness * VolumeScatter);
+float SurfaceBlurAmount = 1.0 - 1.0 * exp(-(SceneDepth - PixelDepth) * SurfaceScatter); // - 0.5 * exp(-(RandomSceneDepth - PixelDepth) * SurfaceScatter);
+float BlurAmount = saturate(VolumeBlurAmount + SurfaceBlurAmount);
+
+float TwoPi = 6.283185;
+float3 CurColor = { 0.0, 0.0, 0.0 };
+float CurDistance = 0.0;
+float Substep = 0.0;
+
+if (BlurAmount == 0.0)
+{
+    return DecodeSceneColorForMaterialNode(NewUV);
+}
+else
+{
+    int i = 0;
+    while (i < RadialSteps)
+    {
+        for (int j = 0; j < AngularSteps; j++)
+        {
+            float Angle = TwoPi * ((TempAARotation + Substep) / AngularSteps);
+            CurOffset.x = cos(Angle);
+            CurOffset.y = sin(Angle);
+            CurOffset *= BlurAmount * MagicBlurNumber * View.ViewToClip[0][0] * 0.01;
+            NewUV = BaseUV + (CurOffset * (CurDistance + TempAADistance));
+            CurColor += DecodeSceneColorForMaterialNode(NewUV);
+            Substep++;
+        }
+        CurDistance += RadialStepSize;
+        Substep += AngularOffset;
+        i++;
+    }
+    CurColor /= (float) (RadialSteps * AngularSteps);
+    return CurColor;
+}
diff --git a/Engine/Plugins/Enterprise/VariantManagerContent/Source/VariantManagerContent/Private/LevelVariantSets.cpp b/Engine/Plugins/Enterprise/VariantManagerContent/Source/VariantManagerContent/Private/LevelVariantSets.cpp
--- a/Engine/Plugins/Enterprise/VariantManagerContent/Source/VariantManagerContent/Private/LevelVariantSets.cpp
+++ b/Engine/Plugins/Enterprise/VariantManagerContent/Source/VariantManagerContent/Private/LevelVariantSets.cpp
@@ -3,8 +3,8 @@
 #include "LevelVariantSets.h"
 
 #include "VariantSet.h"
-#include "Engine/EngineTypes.h"
-#include "Components/ActorComponent.h"
+
+#define LOCTEXT_NAMESPACE "LevelVariantSets"
 
 
 ULevelVariantSets::ULevelVariantSets(const FObjectInitializer& ObjectInitializer)
@@ -12,17 +12,127 @@ ULevelVariantSets::ULevelVariantSets(const FObjectInitializer& ObjectInitializer
 {
 }
 
-void ULevelVariantSets::AddVariantSet(UVariantSet* NewVariantSet)
+void ULevelVariantSets::AddVariantSets(const TArray<UVariantSet*>& NewVariantSets, int32 Index)
 {
-	// Take ownership of NewVariantSet
-	NewVariantSet->Rename(*NewVariantSet->GetName(), this);
+	Modify();
+
+	if (Index == INDEX_NONE)
+	{
+		Index = VariantSets.Num();
+	}
+
+	// Inserting first ensures we preserve the target order
+	VariantSets.Insert(NewVariantSets, Index);
+
+	bool bIsMoveOperation = false;
+	TSet<ULevelVariantSets*> ParentsModified;
+	for (UVariantSet* NewVarSet : NewVariantSets)
+	{
+		ULevelVariantSets* OldParent = NewVarSet->GetParent();
+
+		// We can't just RemoveBinding since that might remove the wrong item in case
+		// we're moving bindings around within this Variant
+		if (OldParent)
+		{
+			if (OldParent != this)
+			{
+				// Don't call RemoveVariantSets here so that we get the entire thing in a single transaction
+				if (!ParentsModified.Contains(OldParent))
+				{
+					OldParent->Modify();
+					ParentsModified.Add(OldParent);
+				}
+				OldParent->VariantSets.RemoveSingle(NewVarSet);
+			}
+			else
+			{
+				bIsMoveOperation = true;
+			}
+		}
+
+		NewVarSet->Modify();
+		NewVarSet->Rename(nullptr, this);  // We'll only actually rename them at the end of this function
+	}
+
+	// If it's a move operation, we'll have to manually clear the old pointers from the array
+	if (bIsMoveOperation)
+	{
+		TSet<UVariantSet*> SetOfNewVariantSets = TSet<UVariantSet*>(NewVariantSets);
 
-	VariantSets.Add(NewVariantSet);
+		// Sweep back from insertion point nulling old bindings with the same GUID
+		for (int32 SweepIndex = Index-1; SweepIndex >= 0; SweepIndex--)
+		{
+			if (SetOfNewVariantSets.Contains(VariantSets[SweepIndex]))
+			{
+				VariantSets[SweepIndex] = nullptr;
+			}
+		}
+		// Sweep forward from the end of the inserted segment nulling old bindings with the same GUID
+		for (int32 SweepIndex = Index + NewVariantSets.Num(); SweepIndex < VariantSets.Num(); SweepIndex++)
+		{
+			if (SetOfNewVariantSets.Contains(VariantSets[SweepIndex]))
+			{
+				VariantSets[SweepIndex] = nullptr;
+			}
+		}
+
+		// Finally remove null entries
+		for (int32 IterIndex = VariantSets.Num() - 1; IterIndex >= 0; IterIndex--)
+		{
+			if (VariantSets[IterIndex] == nullptr)
+			{
+				VariantSets.RemoveAt(IterIndex);
+			}
+		}
+	}
+
+	// Go over new added variants and get them unique display names if they don't have them yet
+	// Can only do this now that we know for sure we deleted the old versions, if we're moving
+	for (UVariantSet* NewVarSet : NewVariantSets)
+	{
+		// Don't transact to keep all of this in a single transaction
+		NewVarSet->SetDisplayText(FText::FromString(GetUniqueVariantSetName(NewVarSet->GetDisplayText().ToString())));
+	}
 }
 
-void ULevelVariantSets::RemoveVariantSet(UVariantSet* ThisVariantSet)
+const TArray<UVariantSet*>& ULevelVariantSets::GetVariantSets() const
 {
-	VariantSets.RemoveSingle(ThisVariantSet);
+	return VariantSets;
+}
+
+void ULevelVariantSets::RemoveVariantSets(const TArray<UVariantSet*> InVariantSets)
+{
+	Modify();
+
+	for (UVariantSet* VariantSet : InVariantSets)
+	{
+		VariantSets.RemoveSingle(VariantSet);
+	}
+}
+
+FString ULevelVariantSets::GetUniqueVariantSetName(const FString& InPrefix)
+{
+	TSet<FString> UniqueNames;
+	for (UVariantSet* VariantSet : VariantSets)
+	{
+		UniqueNames.Add(VariantSet->GetDisplayText().ToString());
+	}
+
+	FString VarSetName = FString(InPrefix);
+
+	// Add a numbered suffix
+	if (UniqueNames.Contains(VarSetName))
+	{
+		int32 Suffix = 0;
+		while (UniqueNames.Contains(VarSetName + FString::FromInt(Suffix)))
+		{
+			Suffix += 1;
+		}
+
+		VarSetName = VarSetName + FString::FromInt(Suffix);
+	}
+
+	return VarSetName;
 }
 
 void ULevelVariantSets::SaveExpansionState(UVariantSet* VarSetOfNode, bool bExpanded)
@@ -37,4 +147,21 @@ bool ULevelVariantSets::GetExpansionState(UVariantSet* VarSetOfNode)
 		return DisplayNodeExpansionStates[VarSetOfNode];
 	}
 	return false;
-}
\ No newline at end of file
+}
+
+int32 ULevelVariantSets::GetNumVariantSets()
+{
+	return VariantSets.Num();
+}
+
+UVariantSet* ULevelVariantSets::GetVariantSet(int32 VariantSetIndex)
+{
+	if (VariantSets.IsValidIndex(VariantSetIndex))
+	{
+		return VariantSets[VariantSetIndex];
+	}
+
+	return nullptr;
+}
+
+#undef LOCTEXT_NAMESPACE
diff --git a/Engine/Plugins/Enterprise/VariantManagerContent/Source/VariantManagerContent/Private/LevelVariantSetsActor.cpp b/Engine/Plugins/Enterprise/VariantManagerContent/Source/VariantManagerContent/Private/LevelVariantSetsActor.cpp
--- a/Engine/Plugins/Enterprise/VariantManagerContent/Source/VariantManagerContent/Private/LevelVariantSetsActor.cpp
+++ b/Engine/Plugins/Enterprise/VariantManagerContent/Source/VariantManagerContent/Private/LevelVariantSetsActor.cpp
@@ -3,6 +3,8 @@
 #include "LevelVariantSetsActor.h"
 
 #include "LevelVariantSets.h"
+#include "VariantSet.h"
+#include "Variant.h"
 
 #include "CoreMinimal.h"
 #include "UObject/ObjectMacros.h"
@@ -47,4 +49,48 @@ ULevelVariantSets* ALevelVariantSetsActor::GetLevelVariantSets(bool bLoad) const
 void ALevelVariantSetsActor::SetLevelVariantSets(ULevelVariantSets* InVariantSets)
 {
     LevelVariantSets = InVariantSets;
+}
+
+bool ALevelVariantSetsActor::SwitchOnVariantByName(FString VariantSetName, FString VariantName)
+{
+	ULevelVariantSets* LevelVarSets = GetLevelVariantSets(true);
+	if (LevelVarSets)
+	{
+		for (UVariantSet* VarSet : LevelVarSets->GetVariantSets())
+		{
+			if (VarSet->GetDisplayText().ToString() == VariantSetName)
+			{
+				for (UVariant* Var : VarSet->GetVariants())
+				{
+					if (Var->GetDisplayText().ToString() == VariantName)
+					{
+						Var->SwitchOn();
+						return true;
+					}
+				}
+			}
+		}
+	}
+
+	return false;
+}
+
+bool ALevelVariantSetsActor::SwitchOnVariantByIndex(int32 VariantSetIndex, int32 VariantIndex)
+{
+	ULevelVariantSets* LevelVarSets = GetLevelVariantSets(true);
+	if (LevelVarSets)
+	{
+		UVariantSet* VarSet = LevelVarSets->GetVariantSet(VariantSetIndex);
+		if (VarSet)
+		{
+			UVariant* Var = VarSet->GetVariant(VariantIndex);
+			if (Var)
+			{
+				Var->SwitchOn();
+				return true;
+			}
+		}
+	}
+
+	return false;
 }
\ No newline at end of file
diff --git a/Engine/Plugins/Enterprise/VariantManagerContent/Source/VariantManagerContent/Private/PropertyValue.cpp b/Engine/Plugins/Enterprise/VariantManagerContent/Source/VariantManagerContent/Private/PropertyValue.cpp
new file mode 100644
--- /dev/null
+++ b/Engine/Plugins/Enterprise/VariantManagerContent/Source/VariantManagerContent/Private/PropertyValue.cpp
@@ -0,0 +1,504 @@
+// Copyright 1998-2018 Epic Games, Inc. All Rights Reserved.
+
+#include "PropertyValue.h"
+
+#include "CoreMinimal.h"
+#include "HAL/UnrealMemory.h"
+#include "VariantObjectBinding.h"
+
+#define LOCTEXT_NAMESPACE "PropertyValue"
+
+DEFINE_LOG_CATEGORY(LogVariantContent);
+
+UPropertyValue::UPropertyValue(const FObjectInitializer& Init)
+{
+}
+
+void UPropertyValue::Init(const TArray<UProperty*>& InProps, const TArray<int32> InIndices, const FString& InFullDisplayString, EPropertyValueCategory InCategory)
+{
+	Properties = InProps;
+	PropertyIndices = InIndices;
+	FullDisplayString = InFullDisplayString;
+	PropCategory = InCategory;
+
+	int32 NumProps = Properties.Num();
+	if (NumProps > 0 && Properties[NumProps-1]->IsA(UObjectProperty::StaticClass()))
+	{
+		bIsObjectProperty = true;
+	}
+
+	ClearLastResolve();
+	ValueBytes.SetNumUninitialized(GetValueSizeInBytes());
+	TempObjPtr.Reset();
+}
+
+const TArray<UProperty*>& UPropertyValue::GetProperties() const
+{
+	return Properties;
+}
+
+const TArray<int32>& UPropertyValue::GetPropertyIndices() const
+{
+	return PropertyIndices;
+}
+
+UVariantObjectBinding* UPropertyValue::GetParent() const
+{
+	return Cast<UVariantObjectBinding>(GetOuter());
+}
+
+void UPropertyValue::Serialize(FArchive& Ar)
+{
+	Super::Serialize(Ar);
+
+	if (Ar.IsSaving())
+	{
+		// Create a soft object ptr with our UObject* data bytes so that it can persist
+		if (bIsObjectProperty && HasRecordedData())
+		{
+			UObject* Obj = *((UObject**)ValueBytes.GetData());
+			if (Obj && Obj->IsValidLowLevel())
+			{
+				TempObjPtr = Obj;
+			}
+		}
+		else
+		{
+			TempObjPtr.Reset();
+		}
+		Ar << TempObjPtr;
+	}
+	else if (Ar.IsLoading())
+	{
+		Ar << TempObjPtr;
+
+		if (bIsObjectProperty)
+		{
+			// Force resolve of our soft object pointer
+			//@Speed Maybe delay this and let it load whenever it wants?
+			UObject* Obj = TempObjPtr.LoadSynchronous();
+
+			if (Obj && Obj->IsValidLowLevel())
+			{
+				int32 NumBytes = sizeof(UObject**);
+				ValueBytes.SetNumUninitialized(NumBytes);
+				FMemory::Memcpy(ValueBytes.GetData(), &Obj, NumBytes);
+				bHasRecordedData = true;
+			}
+			else
+			{
+				TempObjPtr.Reset();
+				bHasRecordedData = false;
+			}
+		}
+	}
+}
+
+bool UPropertyValue::Resolve()
+{
+	UObject* Object = nullptr;
+	UVariantObjectBinding* Parent = GetParent();
+	if (Parent)
+	{
+		Object = Parent->GetObject();
+	}
+
+	if (Object == nullptr)
+	{
+		return false;
+	}
+
+	if (!ResolvePropertiesRecursive(Object->GetClass(), Object, 0))
+	{
+		UE_LOG(LogVariantContent, Error, TEXT("Failed to resolve UPropertyValue '%s' on UObject '%s'"), *GetFullDisplayString(), *Object->GetName());
+		return false;
+	}
+
+	return true;
+}
+
+bool UPropertyValue::HasValidResolve() const
+{
+	return ParentContainerAddress != nullptr;
+}
+
+void UPropertyValue::ClearLastResolve()
+{
+	LeafProperty = nullptr;
+	ParentContainerClass = nullptr;
+	ParentContainerAddress = nullptr;
+	PropertyValuePtr = nullptr;
+}
+
+void* UPropertyValue::GetPropertyParentContainerAddress() const
+{
+	return ParentContainerAddress;
+}
+
+UStruct* UPropertyValue::GetPropertyParentContainerClass() const
+{
+	return ParentContainerClass;
+}
+
+void UPropertyValue::RecordDataFromResolvedObject()
+{
+	if (!Resolve())
+	{
+		return;
+	}
+
+	int32 PropertySizeBytes = GetValueSizeInBytes();
+
+	if (UBoolProperty* PropAsBool = Cast<UBoolProperty>(LeafProperty))
+	{
+		// This could probably be done in a cleaner way since we know its an UBoolProperty...
+		TArray<uint8> BoolBytes;
+		BoolBytes.SetNum(PropertySizeBytes);
+
+		bool* BoolBytesAddress = (bool*)BoolBytes.GetData();
+		*BoolBytesAddress = PropAsBool->GetPropertyValue(PropertyValuePtr);
+
+		SetRecordedData(BoolBytes.GetData(), PropertySizeBytes);
+	}
+	else
+	{
+		SetRecordedData(PropertyValuePtr, PropertySizeBytes);
+	}
+
+	OnPropertyRecorded.Broadcast();
+}
+
+void UPropertyValue::ApplyDataToResolvedObject()
+{
+	if (!HasRecordedData() || !Resolve())
+	{
+		return;
+	}
+
+	// Ready to transact
+	UObject* ContainerOwnerObject = nullptr;
+	UVariantObjectBinding* Parent = GetParent();
+	if (Parent)
+	{
+		ContainerOwnerObject = Parent->GetObject();
+		if (ContainerOwnerObject)
+		{
+			ContainerOwnerObject->SetFlags(RF_Transactional);
+			ContainerOwnerObject->Modify();
+		}
+	}
+	// We might also need to modify a component if we're nested in one
+	UObject* ContainerObject = (UObject*) ParentContainerAddress;
+	if (ContainerObject)
+	{
+		ContainerObject->SetFlags(RF_Transactional);
+		ContainerObject->Modify();
+	}
+
+	// Bool properties need to be set in a particular way since they hold internal private
+	// masks and offsets
+	if (UBoolProperty* PropAsBool = Cast<UBoolProperty>(LeafProperty))
+	{
+		bool* ValueBytesAsBool = (bool*)ValueBytes.GetData();
+		PropAsBool->SetPropertyValue(PropertyValuePtr, *ValueBytesAsBool);
+	}
+	else
+	{
+		// Actually change the object through its property value ptr
+		int32 PropertySizeBytes = LeafProperty->ElementSize;
+		ValueBytes.SetNum(PropertySizeBytes);
+		FMemory::Memcpy(PropertyValuePtr, ValueBytes.GetData(), PropertySizeBytes);
+	}
+
+#if WITH_EDITOR
+	// Update object on viewport
+	if (ContainerObject)
+	{
+		ContainerObject->PostEditChange();
+	}
+	if (ContainerOwnerObject)
+	{
+		ContainerOwnerObject->PostEditChange();
+	}
+#endif
+	OnPropertyApplied.Broadcast();
+}
+
+UClass* UPropertyValue::GetPropertyClass() const
+{
+	UProperty* Prop = GetProperty();
+	if (Prop)
+	{
+		return Prop->GetClass();
+	}
+
+	return nullptr;
+}
+
+UScriptStruct* UPropertyValue::GetStructPropertyStruct() const
+{
+	UProperty* Prop = GetProperty();
+	if (UStructProperty* StructProp = Cast<UStructProperty>(Prop))
+	{
+		return StructProp->Struct;
+	}
+
+	return nullptr;
+}
+
+UClass* UPropertyValue::GetObjectPropertyObjectClass() const
+{
+	if (UObjectProperty* ObjProp = Cast<UObjectProperty>(GetProperty()))
+	{
+		return ObjProp->PropertyClass;
+	}
+
+	return nullptr;
+}
+
+FName UPropertyValue::GetPropertyName() const
+{
+	UProperty* Prop = GetProperty();
+	if (Prop)
+	{
+		return Prop->GetFName();
+	}
+
+	return FName();
+}
+
+const FString& UPropertyValue::GetFullDisplayString() const
+{
+	return FullDisplayString;
+}
+
+FString UPropertyValue::GetLeafDisplayString() const
+{
+	FString LeftString;
+	FString RightString;
+
+	if(FullDisplayString.Split(PATH_DELIMITER, &LeftString, &RightString, ESearchCase::IgnoreCase, ESearchDir::FromEnd))
+	{
+		return RightString;
+	}
+
+	return FullDisplayString;
+}
+
+int32 UPropertyValue::GetValueSizeInBytes() const
+{
+	UProperty* Prop = GetProperty();
+	if (Prop)
+	{
+		return Prop->ElementSize;
+	}
+
+	return 0;
+}
+
+int32 UPropertyValue::GetPropertyOffsetInBytes() const
+{
+	UProperty* Prop = GetProperty();
+	if (Prop)
+	{
+		return Prop->GetOffset_ForInternal();
+	}
+
+	// Dangerous
+	return 0;
+}
+
+bool UPropertyValue::HasRecordedData() const
+{
+	return bHasRecordedData;
+}
+
+const TArray<uint8>& UPropertyValue::GetRecordedData()
+{
+	check(bHasRecordedData)
+
+	ValueBytes.SetNum(GetValueSizeInBytes());
+	return ValueBytes;
+}
+
+void UPropertyValue::SetRecordedData(const uint8* NewDataBytes, int32 NumBytes)
+{
+	Modify();
+
+	if (NumBytes > 0)
+	{
+		ValueBytes.SetNumUninitialized(NumBytes);
+		FMemory::Memcpy(ValueBytes.GetData(), NewDataBytes, NumBytes);
+		bHasRecordedData = true;
+	}
+}
+
+FOnPropertyApplied& UPropertyValue::GetOnPropertyApplied()
+{
+	return OnPropertyApplied;
+}
+
+FOnPropertyRecorded& UPropertyValue::GetOnPropertyRecorded()
+{
+	return OnPropertyRecorded;
+}
+
+UProperty* UPropertyValue::GetProperty() const
+{
+	int32 NumProps = Properties.Num();
+	if (NumProps > 0)
+	{
+		return Properties[NumProps-1];
+	}
+
+	return nullptr;
+}
+
+bool UPropertyValue::ResolvePropertiesRecursive(UStruct* ContainerClass, void* ContainerAddress, int32 SegmentIndex)
+{
+	// Adapted from PropertyPathHelpers.cpp because it is incomplete for arrays of UObjects (important for components)
+
+	UProperty* Property = Properties[SegmentIndex];;
+	const int32 ArrayIndex = PropertyIndices[SegmentIndex] == INDEX_NONE ? 0 : PropertyIndices[SegmentIndex];
+
+	if (SegmentIndex == 0)
+	{
+		ParentContainerClass = ContainerClass;
+		ParentContainerAddress = ContainerAddress;
+	}
+
+	Property = FindField<UProperty>(ContainerClass, Property->GetFName());
+	if (Property)
+	{
+		// Not the last link in the chain --> Dig down deeper updating our class/address if we jump an UObjectProp/UStructProp
+		if (SegmentIndex < (Properties.Num()-1))
+		{
+			// Check first to see if this is a simple object (eg. not an array of objects)
+			if (UObjectProperty* ObjectProperty = Cast<UObjectProperty>(Property))
+			{
+				// If it's an object we need to get the value of the property in the container first before we
+				// can continue, if the object is null we safely stop processing the chain of properties.
+				if ( UObject* CurrentObject = ObjectProperty->GetPropertyValue_InContainer(ContainerAddress, ArrayIndex) )
+				{
+					ParentContainerClass = CurrentObject->GetClass();
+					ParentContainerAddress = CurrentObject;
+
+					return ResolvePropertiesRecursive(CurrentObject->GetClass(), CurrentObject, SegmentIndex + 1);
+				}
+			}
+			// Check to see if this is a simple weak object property (eg. not an array of weak objects).
+			else if (UWeakObjectProperty* WeakObjectProperty = Cast<UWeakObjectProperty>(Property))
+			{
+				FWeakObjectPtr WeakObject = WeakObjectProperty->GetPropertyValue_InContainer(ContainerAddress, ArrayIndex);
+
+				// If it's an object we need to get the value of the property in the container first before we
+				// can continue, if the object is null we safely stop processing the chain of properties.
+				if ( UObject* CurrentObject = WeakObject.Get() )
+				{
+					ParentContainerClass = CurrentObject->GetClass();
+					ParentContainerAddress = CurrentObject;
+
+					return ResolvePropertiesRecursive(CurrentObject->GetClass(), CurrentObject, SegmentIndex + 1);
+				}
+			}
+			// Check to see if this is a simple soft object property (eg. not an array of soft objects).
+			else if (USoftObjectProperty* SoftObjectProperty = Cast<USoftObjectProperty>(Property))
+			{
+				FSoftObjectPtr SoftObject = SoftObjectProperty->GetPropertyValue_InContainer(ContainerAddress, ArrayIndex);
+
+				// If it's an object we need to get the value of the property in the container first before we
+				// can continue, if the object is null we safely stop processing the chain of properties.
+				if ( UObject* CurrentObject = SoftObject.Get() )
+				{
+					ParentContainerClass = CurrentObject->GetClass();
+					ParentContainerAddress = CurrentObject;
+
+					return ResolvePropertiesRecursive(CurrentObject->GetClass(), CurrentObject, SegmentIndex + 1);
+				}
+			}
+			// Check to see if this is a simple structure (eg. not an array of structures)
+			// Note: We don't actually capture properties *inside* UStructProperties, so this path won't be taken. It is here
+			// if we ever wish to change that in the future
+			else if (UStructProperty* StructProp = Cast<UStructProperty>(Property))
+			{
+				void* StructAddress = StructProp->ContainerPtrToValuePtr<void>(ContainerAddress, ArrayIndex);
+
+				ParentContainerClass = StructProp->Struct;
+				ParentContainerAddress = StructAddress;
+
+				return ResolvePropertiesRecursive(StructProp->Struct, StructAddress, SegmentIndex + 1);
+			}
+			else if (UArrayProperty* ArrayProp = Cast<UArrayProperty>(Property))
+			{
+				// We have to replicate these cases in here because we need to access the inner properties
+				// with the FScriptArrayHelper. If we do another recursive call and try parsing the inner
+				// property just as a regular property with an ArrayIndex, it will fail getting the ValuePtr
+				// because for some reason properties always have ArrayDim = 1
+
+				int32 InnerArrayIndex = PropertyIndices[SegmentIndex+1] == INDEX_NONE ? 0 : PropertyIndices[SegmentIndex+1];
+
+				FScriptArrayHelper ArrayHelper(ArrayProp, ArrayProp->ContainerPtrToValuePtr<void>(ContainerAddress));
+				if (!ArrayHelper.IsValidIndex(InnerArrayIndex))
+				{
+					UE_LOG(LogVariantContent, Error, TEXT("Failed to resolve: UArrayProperty '%s' does not have an inner with index %d!"), *ArrayProp->GetName(), InnerArrayIndex);
+					return false;
+				}
+
+				// Array properties show up in the path as two entries (one for the array prop and one for the inner)
+				// so if we're on the second-to-last path segment, it means we want to capture the inner property, so don't
+				// step into it
+				if (SegmentIndex == Properties.Num() - 2)
+				{
+					LeafProperty = ArrayProp->Inner;
+					PropertyValuePtr = ArrayHelper.GetRawPtr(InnerArrayIndex);
+
+					return true;
+				}
+
+				if ( UStructProperty* ArrayOfStructsProp = Cast<UStructProperty>(ArrayProp->Inner) )
+				{
+					void* StructAddress = static_cast<void*>(ArrayHelper.GetRawPtr(InnerArrayIndex));
+
+					ParentContainerClass = ArrayOfStructsProp->Struct;
+					ParentContainerAddress = StructAddress;
+
+					// The next link in the chain is just this array's inner. Let's just skip it instead
+					return ResolvePropertiesRecursive(ArrayOfStructsProp->Struct, StructAddress, SegmentIndex + 2);
+				}
+				if ( UObjectProperty* ObjectProperty = Cast<UObjectProperty>(ArrayProp->Inner) )
+				{
+					void* ObjPtrContainer = ArrayHelper.GetRawPtr(InnerArrayIndex);
+					UObject* CurrentObject = ObjectProperty->GetObjectPropertyValue(ObjPtrContainer);
+					if (CurrentObject)
+					{
+						ParentContainerClass = CurrentObject->GetClass();
+						ParentContainerAddress =  CurrentObject;
+
+						// The next link in the chain is just this array's inner. Let's just skip it instead
+						return ResolvePropertiesRecursive(CurrentObject->GetClass(), CurrentObject, SegmentIndex + 2);
+					}
+				}
+			}
+			else if( USetProperty* SetProperty = Cast<USetProperty>(Property) )
+			{
+				// TODO: we dont support set properties yet
+			}
+			else if( UMapProperty* MapProperty = Cast<UMapProperty>(Property) )
+			{
+				// TODO: we dont support map properties yet
+			}
+		}
+		// Last link, the thing we actually want to capture
+		else
+		{
+			LeafProperty = Property;
+			PropertyValuePtr = LeafProperty->ContainerPtrToValuePtr<uint8>(ContainerAddress, ArrayIndex);
+			return true;
+		}
+	}
+
+	ClearLastResolve();
+	return false;
+}
+
+#undef LOCTEXT_NAMESPACE
diff --git a/Engine/Plugins/Enterprise/VariantManagerContent/Source/VariantManagerContent/Private/PropertyValueMaterial.cpp b/Engine/Plugins/Enterprise/VariantManagerContent/Source/VariantManagerContent/Private/PropertyValueMaterial.cpp
new file mode 100644
--- /dev/null
+++ b/Engine/Plugins/Enterprise/VariantManagerContent/Source/VariantManagerContent/Private/PropertyValueMaterial.cpp
@@ -0,0 +1,149 @@
+// Copyright 1998-2018 Epic Games, Inc. All Rights Reserved.
+
+#include "PropertyValueMaterial.h"
+
+#include "CoreMinimal.h"
+#include "Components/MeshComponent.h"
+#include "HAL/UnrealMemory.h"
+#include "VariantObjectBinding.h"
+
+#define LOCTEXT_NAMESPACE "PropertyValueMaterial"
+
+
+UPropertyValueMaterial::UPropertyValueMaterial(const FObjectInitializer& Init)
+	: Super(Init)
+{
+}
+
+bool UPropertyValueMaterial::Resolve()
+{
+	UObject* Object = nullptr;
+	UVariantObjectBinding* Parent = GetParent();
+	if (Parent)
+	{
+		Object = Parent->GetObject();
+	}
+
+	if (Object == nullptr)
+	{
+		return false;
+	}
+
+	// Remove the last couple of links in the chain because they might not resolve
+	// These will be UArrayProperty OverrideMaterials and its Inner
+	UProperty* OverrideInner = Properties.Pop();
+	UProperty* OverrideOuter = Properties.Pop();
+	int32 OverrideInnerIndex = PropertyIndices.Pop();
+	int32 OverrideOuterIndex = PropertyIndices.Pop();
+
+	if (!ResolvePropertiesRecursive(Object->GetClass(), Object, 0))
+	{
+		UE_LOG(LogVariantContent, Error, TEXT("Failed to resolve UPropertyValueMaterial '%s' on UObject '%s'"), *GetFullDisplayString(), *Object->GetName());
+		return false;
+	}
+
+	Properties.Add(OverrideOuter);
+	Properties.Add(OverrideInner);
+	PropertyIndices.Add(OverrideOuterIndex);
+	PropertyIndices.Add(OverrideInnerIndex);
+
+	if (PropertyValuePtr == nullptr)
+	{
+		UE_LOG(LogVariantContent, Error, TEXT("UPropertyValueMaterial '%s' does not target a valid UStaticMeshComponent!"), *GetFullDisplayString());
+		return false;
+	}
+
+	ParentContainerAddress = *((UObject**)PropertyValuePtr);
+	ParentContainerClass = ((UObject*)ParentContainerAddress)->GetClass();
+
+	// We don't want anything trying to access this property by itself
+	PropertyValuePtr = nullptr;
+	LeafProperty = Properties.Num() > 0 ? Properties[Properties.Num()-1] : nullptr;
+	return true;
+}
+
+void UPropertyValueMaterial::RecordDataFromResolvedObject()
+{
+	if (!Resolve())
+	{
+		return;
+	}
+
+	int32 PropertySizeBytes = GetValueSizeInBytes();
+
+	UMeshComponent* ContainerObject = (UMeshComponent*) ParentContainerAddress;
+	if (!ContainerObject)
+	{
+		UE_LOG(LogVariantContent, Error, TEXT("UPropertyValueMaterial '%s' does not have a UMeshComponent as parent address!"), *GetFullDisplayString());
+		return;
+	}
+
+	int32 NumIndices = PropertyIndices.Num();
+	if (NumIndices < 1)
+	{
+		return;
+	}
+
+	int32 MatIndex = PropertyIndices[NumIndices-1];
+	UMaterialInterface* Mat = ContainerObject->GetMaterial(MatIndex);
+
+	if (Mat && Mat->IsValidLowLevel())
+	{
+		SetRecordedData((uint8*)&Mat, PropertySizeBytes);
+	}
+
+	OnPropertyRecorded.Broadcast();
+}
+
+void UPropertyValueMaterial::ApplyDataToResolvedObject()
+{
+	if (!HasRecordedData() || !Resolve())
+	{
+		return;
+	}
+
+	// Ready to transact
+	UObject* ContainerOwnerObject = nullptr;
+	UVariantObjectBinding* Parent = GetParent();
+	if (Parent)
+	{
+		ContainerOwnerObject = Parent->GetObject();
+		if (ContainerOwnerObject)
+		{
+			ContainerOwnerObject->SetFlags(RF_Transactional);
+			ContainerOwnerObject->Modify();
+		}
+	}
+
+	UMeshComponent* ContainerObject = (UMeshComponent*) ParentContainerAddress;
+	if (!ContainerObject)
+	{
+		UE_LOG(LogVariantContent, Error, TEXT("UPropertyValueMaterial '%s' does not have a UMeshComponent as parent address!"), *GetFullDisplayString());
+		return;
+	}
+
+	ContainerObject->SetFlags(RF_Transactional);
+	ContainerObject->Modify();
+
+	UMaterialInterface* Mat = *((UMaterialInterface**)ValueBytes.GetData());
+
+	int32 NumIndices = PropertyIndices.Num();
+	if (Mat && Mat->IsValidLowLevel() && NumIndices > 0)
+	{
+		int32 MatIndex = PropertyIndices[NumIndices-1];
+		ContainerObject->SetMaterial(MatIndex, Mat);
+	}
+
+	// Update object on viewport
+#if WITH_EDITOR
+	ContainerObject->PostEditChange();
+	if (ContainerOwnerObject)
+	{
+		ContainerOwnerObject->PostEditChange();
+	}
+#endif
+
+	OnPropertyApplied.Broadcast();
+}
+
+#undef LOCTEXT_NAMESPACE
diff --git a/Engine/Plugins/Enterprise/VariantManagerContent/Source/VariantManagerContent/Private/PropertyValueTransform.cpp b/Engine/Plugins/Enterprise/VariantManagerContent/Source/VariantManagerContent/Private/PropertyValueTransform.cpp
new file mode 100644
--- /dev/null
+++ b/Engine/Plugins/Enterprise/VariantManagerContent/Source/VariantManagerContent/Private/PropertyValueTransform.cpp
@@ -0,0 +1,85 @@
+// Copyright 1998-2018 Epic Games, Inc. All Rights Reserved.
+
+#include "PropertyValueTransform.h"
+
+#include "CoreMinimal.h"
+#include "Components/SceneComponent.h"
+#include "HAL/UnrealMemory.h"
+#include "VariantObjectBinding.h"
+
+#define LOCTEXT_NAMESPACE "PropertyValueTransform"
+
+
+UPropertyValueTransform::UPropertyValueTransform(const FObjectInitializer& Init)
+	: Super(Init)
+{
+}
+
+void UPropertyValueTransform::ApplyDataToResolvedObject()
+{
+	if (!HasRecordedData() || !Resolve())
+	{
+		return;
+	}
+
+	// Ready to transact
+	UObject* ContainerOwnerObject = nullptr;
+	UVariantObjectBinding* Parent = GetParent();
+	if (Parent)
+	{
+		ContainerOwnerObject = Parent->GetObject();
+		if (ContainerOwnerObject)
+		{
+			ContainerOwnerObject->SetFlags(RF_Transactional);
+			ContainerOwnerObject->Modify();
+		}
+	}
+
+	USceneComponent* ContainerObject = (USceneComponent*) ParentContainerAddress;
+	if (!ContainerObject)
+	{
+		UE_LOG(LogVariantContent, Error, TEXT("UPropertyValueTransform '%s' does not have a USceneComponent as parent address!"), *GetFullDisplayString());
+		return;
+	}
+
+	ContainerObject->SetFlags(RF_Transactional);
+	ContainerObject->Modify();
+
+	switch (PropCategory)
+	{
+	case EPropertyValueCategory::RelativeLocation:
+	{
+		FVector BytesAsVec = *(FVector*)ValueBytes.GetData();
+		ContainerObject->SetRelativeLocation(BytesAsVec);
+		break;
+	}
+	case EPropertyValueCategory::RelativeRotation:
+	{
+		FRotator BytesAsRot = *(FRotator*)ValueBytes.GetData();
+		ContainerObject->SetRelativeRotation(BytesAsRot);
+		break;
+	}
+	case EPropertyValueCategory::RelativeScale3D:
+	{
+		FVector BytesAsVec = *(FVector*)ValueBytes.GetData();
+		ContainerObject->SetRelativeScale3D(BytesAsVec);
+		break;
+	}
+	default:
+		check(false);
+		break;
+	}
+
+	// Update object on viewport
+#if WITH_EDITOR
+	ContainerObject->PostEditChange();
+	if (ContainerOwnerObject)
+	{
+		ContainerOwnerObject->PostEditChange();
+	}
+#endif
+
+	OnPropertyApplied.Broadcast();
+}
+
+#undef LOCTEXT_NAMESPACE
diff --git a/Engine/Plugins/Enterprise/VariantManagerContent/Source/VariantManagerContent/Private/PropertyValueVisibility.cpp b/Engine/Plugins/Enterprise/VariantManagerContent/Source/VariantManagerContent/Private/PropertyValueVisibility.cpp
new file mode 100644
--- /dev/null
+++ b/Engine/Plugins/Enterprise/VariantManagerContent/Source/VariantManagerContent/Private/PropertyValueVisibility.cpp
@@ -0,0 +1,63 @@
+// Copyright 1998-2018 Epic Games, Inc. All Rights Reserved.
+
+#include "PropertyValueVisibility.h"
+
+#include "CoreMinimal.h"
+#include "Components/SceneComponent.h"
+#include "HAL/UnrealMemory.h"
+#include "VariantObjectBinding.h"
+
+#define LOCTEXT_NAMESPACE "PropertyValueVisibility"
+
+
+UPropertyValueVisibility::UPropertyValueVisibility(const FObjectInitializer& Init)
+	: Super(Init)
+{
+}
+
+void UPropertyValueVisibility::ApplyDataToResolvedObject()
+{
+	if (!HasRecordedData() || !Resolve())
+	{
+		return;
+	}
+
+	// Ready to transact
+	UObject* ContainerOwnerObject = nullptr;
+	UVariantObjectBinding* Parent = GetParent();
+	if (Parent)
+	{
+		ContainerOwnerObject = Parent->GetObject();
+		if (ContainerOwnerObject)
+		{
+			ContainerOwnerObject->SetFlags(RF_Transactional);
+			ContainerOwnerObject->Modify();
+		}
+	}
+
+	USceneComponent* ContainerObject = (USceneComponent*) ParentContainerAddress;
+	if (!ContainerObject)
+	{
+		UE_LOG(LogVariantContent, Error, TEXT("UPropertyValueVisibility '%s' does not have a USceneComponent as parent address!"), *GetFullDisplayString());
+		return;
+	}
+
+	ContainerObject->SetFlags(RF_Transactional);
+	ContainerObject->Modify();
+
+	const bool bVisibilityValue = *(ValueBytes.GetData()) != 0;
+	ContainerObject->SetVisibility(bVisibilityValue, true);
+
+	// Update object on viewport
+#if WITH_EDITOR
+	ContainerObject->PostEditChange();
+	if (ContainerOwnerObject)
+	{
+		ContainerOwnerObject->PostEditChange();
+	}
+#endif
+
+	OnPropertyApplied.Broadcast();
+}
+
+#undef LOCTEXT_NAMESPACE
diff --git a/Engine/Plugins/Enterprise/VariantManagerContent/Source/VariantManagerContent/Private/Variant.cpp b/Engine/Plugins/Enterprise/VariantManagerContent/Source/VariantManagerContent/Private/Variant.cpp
--- a/Engine/Plugins/Enterprise/VariantManagerContent/Source/VariantManagerContent/Private/Variant.cpp
+++ b/Engine/Plugins/Enterprise/VariantManagerContent/Source/VariantManagerContent/Private/Variant.cpp
@@ -1,123 +1,161 @@
 // Copyright 1998-2018 Epic Games, Inc. All Rights Reserved.
 
 #include "Variant.h"
+
+#include "PropertyValue.h"
 #include "VariantSet.h"
 #include "VariantObjectBinding.h"
-
 #include "CoreMinimal.h"
 #include "GameFramework/Actor.h"
 
-#define LOCTEXT_NAMESPACE "VariantManagerVariant"
+#define LOCTEXT_NAMESPACE "Variant"
 
 
 UVariant::UVariant(const FObjectInitializer& Init)
 {
+	DisplayText = FText::FromString(TEXT("Variant"));
 }
 
-UVariant* UVariant::Clone(UObject* ClonesOuter)
+UVariantSet* UVariant::GetParent()
 {
-	if (ClonesOuter == INVALID_OBJECT)
-	{
-		ClonesOuter = GetOuter();
-	}
-
-	UVariant* NewVariant = DuplicateObject(this, ClonesOuter);
-
-	for (UVariantObjectBinding* OurBinding : GetBindings())
-	{
-		NewVariant->AddBinding(OurBinding->Clone());
-	}
-	NewVariant->SetSortingOrder(GetSortingOrder() + 1);
+	return Cast<UVariantSet>(GetOuter());
+}
 
-	return NewVariant;
+void UVariant::SetDisplayText(const FText& NewDisplayText)
+{
+	DisplayText = NewDisplayText;
 }
 
-UVariantSet* UVariant::GetParent()
+FText UVariant::GetDisplayText() const
 {
-	return Cast<UVariantSet>(GetOuter());
+	return DisplayText;
 }
 
-FText UVariant::GetDisplayName() const
+void UVariant::AddBindings(const TArray<UVariantObjectBinding*>& NewBindings, int32 Index)
 {
-	if (DisplayName.IsEmpty())
+	Modify();
+
+	if (Index == INDEX_NONE)
 	{
-		return GetDefaultDisplayName();
+		Index = ObjectBindings.Num();
 	}
 
-	return DisplayName;
-}
+	// Inserting first ensures we preserve the target order
+	ObjectBindings.Insert(NewBindings, Index);
 
-void UVariant::SetDisplayName(const FText& NewDisplayName)
-{
-	if (NewDisplayName.EqualTo(DisplayName))
+	bool bIsMoveOperation = false;
+	TSet<UVariant*> ParentsModified;
+	for (UVariantObjectBinding* NewBinding : NewBindings)
+	{
+		UVariant* OldParent = NewBinding->GetParent();
+
+		// We can't just RemoveBinding since that might remove the wrong item in case
+		// we're moving bindings around within this Variant
+		if (OldParent)
+		{
+			if (OldParent != this)
+			{
+				// Don't call RemoveBinding here so that we get the entire thing in a single transaction
+				if (!ParentsModified.Contains(OldParent))
+				{
+					OldParent->Modify();
+					ParentsModified.Add(OldParent);
+				}
+				OldParent->ObjectBindings.RemoveSingle(NewBinding);
+			}
+			else
+			{
+				bIsMoveOperation = true;
+			}
+		}
+
+		NewBinding->Modify();
+		NewBinding->Rename(nullptr, this);
+	}
+
+	// If it's a move operation, we'll have to manually clear the old pointers from the array
+	if (!bIsMoveOperation)
 	{
 		return;
 	}
 
-	SetFlags(RF_Transactional);
-	Modify();
+	TSet<FString> NewBindingPaths = TSet<FString>();
+	for (UVariantObjectBinding* NewBinding : NewBindings)
+	{
+		NewBindingPaths.Add(NewBinding->GetObjectPath());
+	}
+
+	// Sweep back from insertion point nulling old bindings with the same GUID
+	for (int32 SweepIndex = Index-1; SweepIndex >= 0; SweepIndex--)
+	{
+		if (NewBindingPaths.Contains(ObjectBindings[SweepIndex]->GetObjectPath()))
+		{
+			ObjectBindings[SweepIndex] = nullptr;
+		}
+	}
+	// Sweep forward from the end of the inserted segment nulling old bindings with the same GUID
+	for (int32 SweepIndex = Index + NewBindings.Num(); SweepIndex < ObjectBindings.Num(); SweepIndex++)
+	{
+		if (NewBindingPaths.Contains(ObjectBindings[SweepIndex]->GetObjectPath()))
+		{
+			ObjectBindings[SweepIndex] = nullptr;
+		}
+	}
 
-	DisplayName = NewDisplayName;
+	// Finally remove null entries
+	for (int32 IterIndex = ObjectBindings.Num() - 1; IterIndex >= 0; IterIndex--)
+	{
+		if (ObjectBindings[IterIndex] == nullptr)
+		{
+			ObjectBindings.RemoveAt(IterIndex);
+		}
+	}
 }
 
-FText UVariant::GetDefaultDisplayName() const
+const TArray<UVariantObjectBinding*>& UVariant::GetBindings()
 {
-	return LOCTEXT("UnnamedVariantName", "Variant");
+	return ObjectBindings;
 }
 
-void UVariant::AddBinding(UVariantObjectBinding* NewBinding)
+void UVariant::RemoveBindings(const TArray<UVariantObjectBinding*>& Bindings)
 {
-	NewBinding->Rename(*NewBinding->GetName(), this);
+	Modify();
 
-	ObjectBindings.Add(NewBinding);
+	for (UVariantObjectBinding* Binding : Bindings)
+	{
+		ObjectBindings.RemoveSingle(Binding);
+	}
 }
 
-void UVariant::AddActors(TWeakObjectPtr<AActor> InActor)
+int32 UVariant::GetNumActors()
 {
-	UVariantObjectBinding* NewBinding = NewObject<UVariantObjectBinding>(this);
-	NewBinding->Init(InActor.Get());
+	return ObjectBindings.Num();
+}
 
-	FGuid NewGuid = NewBinding->GetObjectGuid();
-	for (UVariantObjectBinding* Binding : ObjectBindings)
+AActor* UVariant::GetActor(int32 ActorIndex)
+{
+	if (ObjectBindings.IsValidIndex(ActorIndex))
 	{
-		if (Binding->GetObjectGuid() == NewGuid)
+		UVariantObjectBinding* Binding = ObjectBindings[ActorIndex];
+		UObject* Obj = Binding->GetObject();
+		if (AActor* Actor = Cast<AActor>(Obj))
 		{
-			return;
+			return Actor;
 		}
 	}
 
-	AddBinding(NewBinding);
+	return nullptr;
 }
 
-void UVariant::AddActors(const TArray<TWeakObjectPtr<AActor>>& InActors)
+void UVariant::SwitchOn()
 {
-	for (TWeakObjectPtr<AActor> InActor : InActors)
+	for (UVariantObjectBinding* Binding : ObjectBindings)
 	{
-		UVariantObjectBinding* NewBinding = NewObject<UVariantObjectBinding>(this);
-		NewBinding->Init(InActor.Get());
-
-		FGuid NewGuid = NewBinding->GetObjectGuid();
-
-		bool bDuplicate = false;
-		for (UVariantObjectBinding* Binding : ObjectBindings)
+		for (UPropertyValue* PropCapture : Binding->GetCapturedProperties())
 		{
-			if (Binding->GetObjectGuid() == NewGuid)
-			{
-				bDuplicate = true;
-				break;
-			}
-		}
-
-		if (!bDuplicate)
-		{
-			AddBinding(NewBinding);
+			PropCapture->ApplyDataToResolvedObject();
 		}
 	}
 }
 
-void UVariant::RemoveBinding(UVariantObjectBinding* ThisBinding)
-{
-	ObjectBindings.RemoveSingle(ThisBinding);
-}
-
 #undef LOCTEXT_NAMESPACE
\ No newline at end of file
diff --git a/Engine/Plugins/Enterprise/VariantManagerContent/Source/VariantManagerContent/Private/VariantObjectBinding.cpp b/Engine/Plugins/Enterprise/VariantManagerContent/Source/VariantManagerContent/Private/VariantObjectBinding.cpp
--- a/Engine/Plugins/Enterprise/VariantManagerContent/Source/VariantManagerContent/Private/VariantObjectBinding.cpp
+++ b/Engine/Plugins/Enterprise/VariantManagerContent/Source/VariantManagerContent/Private/VariantObjectBinding.cpp
@@ -2,29 +2,168 @@
 
 #include "VariantObjectBinding.h"
 
+#include "UObject/LazyObjectPtr.h"
+#include "PropertyValue.h"
 #include "Variant.h"
-#include "MovieSceneTrack.h"
+
+#define LOCTEXT_NAMESPACE "VariantObjectBinding"
 
 
 UVariantObjectBinding::UVariantObjectBinding(const FObjectInitializer& Init)
 {
+}
 
+void UVariantObjectBinding::Init(UObject* InObject)
+{
+	ObjectPtr = InObject;
+	LazyObjectPtr = InObject;
+}
+
+UVariant* UVariantObjectBinding::GetParent()
+{
+	return Cast<UVariant>(GetOuter());
 }
 
-UVariantObjectBinding* UVariantObjectBinding::Clone(UObject* ClonesOuter)
+FText UVariantObjectBinding::GetDisplayText() const
 {
-	if (ClonesOuter == INVALID_OBJECT)
+	UObject* Obj = GetObject();
+	if (Obj)
+	{
+		return FText::FromName(Obj->GetFName());
+	}
+	else
 	{
-		ClonesOuter = GetOuter();
+		return FText::FromString(TEXT("<Unloaded binding>"));
 	}
+}
 
-	UVariantObjectBinding* NewBinding = DuplicateObject(this, ClonesOuter);
-	NewBinding->Init(GetObject(), GetSortingOrder() + 1);
+FString UVariantObjectBinding::GetObjectPath() const
+{
+	return ObjectPtr.ToString();
+}
 
-	return NewBinding;
+UObject* UVariantObjectBinding::GetObject() const
+{
+	if (ObjectPtr.IsValid())
+	{
+		UObject* Obj = ObjectPtr.ResolveObject();
+		if (Obj)
+		{
+			return Obj;
+		}
+
+		// Last resort: When an actor switches levels (e.g. new level -> create variant manager ->
+		// -> add one of the temp actors -> save level) the FSoftObjectPath might lose track of it
+		// We will then use our TLazyObjectPtr to try and find the object since it uses FGuids and
+		// it might still have a valid link
+		if (!Obj)
+		{
+			UObject* LazyObject = LazyObjectPtr.Get();
+			if (LazyObject)
+			{
+				UE_LOG(LogVariantContent, Log, TEXT("Actor '%s' switched path. Binding updating from path '%s' to '%s'"), *LazyObject->GetName(), *ObjectPtr.ToString(), *LazyObject->GetFullName());
+				ObjectPtr = LazyObject;
+			}
+		}
+	}
+
+	return nullptr;
 }
 
-UVariant* UVariantObjectBinding::GetParent()
+void UVariantObjectBinding::AddCapturedProperties(const TArray<UPropertyValue*>& NewProperties, int32 Index)
 {
-	return Cast<UVariant>(GetOuter());
+	Modify();
+
+	if (Index == INDEX_NONE)
+	{
+		Index = CapturedProperties.Num();
+	}
+
+	// Inserting first ensures we preserve the target order
+	CapturedProperties.Insert(NewProperties, Index);
+
+	bool bIsMoveOperation = false;
+	TSet<UVariantObjectBinding*> ParentsModified;
+	for (UPropertyValue* NewProp : NewProperties)
+	{
+		UVariantObjectBinding* OldParent = NewProp->GetParent();
+
+		// We can't just RemoveCapturedProperty since that might remove the wrong item in case
+		// we're moving bindings around within this Variant
+		if (OldParent)
+		{
+			if (OldParent != this)
+			{
+				// Don't call RemoveProperty here so that we get the entire thing in a single transaction
+				if (!ParentsModified.Contains(OldParent))
+				{
+					OldParent->Modify();
+					ParentsModified.Add(OldParent);
+				}
+				OldParent->CapturedProperties.RemoveSingle(NewProp);
+			}
+			else
+			{
+				bIsMoveOperation = true;
+			}
+		}
+
+		NewProp->Modify();
+		NewProp->Rename(nullptr, this);
+	}
+
+	// If it's a move operation, we'll have to manually clear the old pointers from the array
+	if (!bIsMoveOperation)
+	{
+		return;
+	}
+
+	TSet<FString> NewPropertyPaths = TSet<FString>();
+	for (UPropertyValue* NewProp : NewProperties)
+	{
+		NewPropertyPaths.Add(NewProp->GetFullDisplayString());
+	}
+
+	// Sweep back from insertion point nulling old bindings with the same GUID
+	for (int32 SweepIndex = Index-1; SweepIndex >= 0; SweepIndex--)
+	{
+		if (NewPropertyPaths.Contains(CapturedProperties[SweepIndex]->GetFullDisplayString()))
+		{
+			CapturedProperties[SweepIndex] = nullptr;
+		}
+	}
+	// Sweep forward from the end of the inserted segment nulling old bindings with the same GUID
+	for (int32 SweepIndex = Index + NewProperties.Num(); SweepIndex < CapturedProperties.Num(); SweepIndex++)
+	{
+		if (NewPropertyPaths.Contains(CapturedProperties[SweepIndex]->GetFullDisplayString()))
+		{
+			CapturedProperties[SweepIndex] = nullptr;
+		}
+	}
+
+	// Finally remove null entries
+	for (int32 IterIndex = CapturedProperties.Num() - 1; IterIndex >= 0; IterIndex--)
+	{
+		if (CapturedProperties[IterIndex] == nullptr)
+		{
+			CapturedProperties.RemoveAt(IterIndex);
+		}
+	}
 }
+
+const TArray<UPropertyValue*>& UVariantObjectBinding::GetCapturedProperties() const
+{
+	return CapturedProperties;
+}
+
+void UVariantObjectBinding::RemoveCapturedProperties(const TArray<UPropertyValue*>& Properties)
+{
+	Modify();
+
+	for (UPropertyValue* Prop : Properties)
+	{
+		CapturedProperties.RemoveSingle(Prop);
+	}
+}
+
+#undef LOCTEXT_NAMESPACE
diff --git a/Engine/Plugins/Enterprise/VariantManagerContent/Source/VariantManagerContent/Private/VariantSet.cpp b/Engine/Plugins/Enterprise/VariantManagerContent/Source/VariantManagerContent/Private/VariantSet.cpp
--- a/Engine/Plugins/Enterprise/VariantManagerContent/Source/VariantManagerContent/Private/VariantSet.cpp
+++ b/Engine/Plugins/Enterprise/VariantManagerContent/Source/VariantManagerContent/Private/VariantSet.cpp
@@ -1,9 +1,9 @@
 // Copyright 1998-2018 Epic Games, Inc. All Rights Reserved.
 
 #include "VariantSet.h"
+
 #include "LevelVariantSets.h"
 #include "Variant.h"
-
 #include "CoreMinimal.h"
 
 #define LOCTEXT_NAMESPACE "VariantManagerVariantSet"
@@ -11,69 +11,160 @@
 
 UVariantSet::UVariantSet(const FObjectInitializer& Init)
 {
+	DisplayText = FText::FromString(TEXT("VariantSet"));
+}
+
+ULevelVariantSets* UVariantSet::GetParent()
+{
+	return Cast<ULevelVariantSets>(GetOuter());
+}
+
+void UVariantSet::SetDisplayText(const FText& NewDisplayText)
+{
+	DisplayText = NewDisplayText;
 }
 
-UVariantSet* UVariantSet::Clone(UObject* ClonesOuter)
+FText UVariantSet::GetDisplayText() const
 {
-	if (ClonesOuter == INVALID_OBJECT)
+	return DisplayText;
+}
+
+FString UVariantSet::GetUniqueVariantName(const FString& InPrefix)
+{
+	TSet<FString> UniqueNames;
+	for (UVariant* Variant : Variants)
 	{
-		ClonesOuter = GetOuter();
+		UniqueNames.Add(Variant->GetDisplayText().ToString());
 	}
 
-	UVariantSet* NewVariantSet = DuplicateObject(this, GetOuter());
+	FString VarName = FString(InPrefix);
 
-	for (UVariant* OurVariant : GetVariants())
+	// Add a numbered suffix
+	if (UniqueNames.Contains(VarName))
 	{
-		NewVariantSet->AddVariant(OurVariant->Clone());
+		int32 Suffix = 0;
+		while (UniqueNames.Contains(VarName + FString::FromInt(Suffix)))
+		{
+			Suffix += 1;
+		}
+
+		VarName = VarName + FString::FromInt(Suffix);
 	}
-	NewVariantSet->SetSortingOrder(GetSortingOrder() + 1);
 
-	return NewVariantSet;
+	return VarName;
 }
 
-ULevelVariantSets* UVariantSet::GetParent()
+void UVariantSet::AddVariants(const TArray<UVariant*>& NewVariants, int32 Index)
 {
-	return Cast<ULevelVariantSets>(GetOuter());
-}
+	Modify();
 
-FText UVariantSet::GetDisplayName() const
-{
-	if (DisplayName.IsEmpty())
+	if (Index == INDEX_NONE)
 	{
-		return GetDefaultDisplayName();
+		Index = Variants.Num();
 	}
 
-	return DisplayName;
-}
+	// Inserting first ensures we preserve the target order
+	Variants.Insert(NewVariants, Index);
 
-void UVariantSet::SetDisplayName(const FText& NewDisplayName)
-{
-	if (NewDisplayName.EqualTo(DisplayName))
+	bool bIsMoveOperation = false;
+	TSet<UVariantSet*> ParentsModified;
+	for (UVariant* NewVariant : NewVariants)
 	{
-		return;
+		UVariantSet* OldParent = NewVariant->GetParent();
+
+		// We can't just RemoveBinding since that might remove the wrong item in case
+		// we're moving bindings around within this Variant
+		if (OldParent)
+		{
+			if (OldParent != this)
+			{
+				// Don't call RemoveVariants here so that we get the entire thing in a single transaction
+				if (!ParentsModified.Contains(OldParent))
+				{
+					OldParent->Modify();
+					ParentsModified.Add(OldParent);
+				}
+				OldParent->Variants.RemoveSingle(NewVariant);
+			}
+			else
+			{
+				bIsMoveOperation = true;
+			}
+		}
+
+		NewVariant->Modify();
+		NewVariant->Rename(nullptr, this);  // We'll only actually rename them at the end of this function
 	}
 
-	SetFlags(RF_Transactional);
-	Modify();
+	// If it's a move operation, we'll have to manually clear the old pointers from the array
+	if (bIsMoveOperation)
+	{
+		TSet<UVariant*> SetOfNewVariants = TSet<UVariant*>(NewVariants);
+
+		// Sweep back from insertion point nulling old bindings with the same GUID
+		for (int32 SweepIndex = Index-1; SweepIndex >= 0; SweepIndex--)
+		{
+			if (SetOfNewVariants.Contains(Variants[SweepIndex]))
+			{
+				Variants[SweepIndex] = nullptr;
+			}
+		}
+		// Sweep forward from the end of the inserted segment nulling old bindings with the same GUID
+		for (int32 SweepIndex = Index + NewVariants.Num(); SweepIndex < Variants.Num(); SweepIndex++)
+		{
+			if (SetOfNewVariants.Contains(Variants[SweepIndex]))
+			{
+				Variants[SweepIndex] = nullptr;
+			}
+		}
+
+		// Finally remove null entries
+		for (int32 IterIndex = Variants.Num() - 1; IterIndex >= 0; IterIndex--)
+		{
+			if (Variants[IterIndex] == nullptr)
+			{
+				Variants.RemoveAt(IterIndex);
+			}
+		}
+	}
 
-	DisplayName = NewDisplayName;
+	// Go over new added variants and get them unique display names if they don't have them yet
+	// Can only do this now that we know for sure we deleted the old versions, if we're moving
+	for (UVariant* NewVar : NewVariants)
+	{
+		// Don't transact to keep all of this in a single transaction
+		NewVar->SetDisplayText(FText::FromString(GetUniqueVariantName(NewVar->GetDisplayText().ToString())));
+	}
 }
 
-FText UVariantSet::GetDefaultDisplayName() const
+const TArray<UVariant*>& UVariantSet::GetVariants() const
 {
-	return LOCTEXT("UnnamedVariantSetName", "VariantSet");
+	return Variants;
 }
 
-void UVariantSet::AddVariant(UVariant* NewVariant)
+void UVariantSet::RemoveVariants(const TArray<UVariant*>& InVariants)
 {
-	NewVariant->Rename(*NewVariant->GetName(), this);
+	Modify();
 
-	Variants.Add(NewVariant);
+	for (UVariant* Variant : InVariants)
+	{
+		Variants.RemoveSingle(Variant);
+	}
+}
+
+int32 UVariantSet::GetNumVariants()
+{
+	return Variants.Num();
 }
 
-void UVariantSet::RemoveVariant(UVariant* ThisVariant)
+UVariant* UVariantSet::GetVariant(int32 VariantIndex)
 {
-	Variants.RemoveSingle(ThisVariant);
+	if (Variants.IsValidIndex(VariantIndex))
+	{
+		return Variants[VariantIndex];
+	}
+
+	return nullptr;
 }
 
 #undef LOCTEXT_NAMESPACE
\ No newline at end of file
diff --git a/Engine/Plugins/Enterprise/VariantManagerContent/Source/VariantManagerContent/Public/LevelVariantSets.h b/Engine/Plugins/Enterprise/VariantManagerContent/Source/VariantManagerContent/Public/LevelVariantSets.h
--- a/Engine/Plugins/Enterprise/VariantManagerContent/Source/VariantManagerContent/Public/LevelVariantSets.h
+++ b/Engine/Plugins/Enterprise/VariantManagerContent/Source/VariantManagerContent/Public/LevelVariantSets.h
@@ -18,19 +18,21 @@ class VARIANTMANAGERCONTENT_API ULevelVariantSets : public UObject
 
 public:
 
-	/**~ UObject implementation */
-	//virtual void Serialize(FArchive& Ar) override;
+	void AddVariantSets(const TArray<UVariantSet*>& NewVariantSets, int32 Index = INDEX_NONE);
+	const TArray<UVariantSet*>& GetVariantSets() const;
+	void RemoveVariantSets(const TArray<UVariantSet*> InVariantSets);
 
-	void AddVariantSet(UVariantSet* NewVariantSet);
-	const TArray<UVariantSet*>& GetVariantSets() const
-	{
-		return VariantSets;
-	}
-	void RemoveVariantSet(UVariantSet* ThisVariantSet);
+	FString GetUniqueVariantSetName(const FString& Prefix);
 
 	void SaveExpansionState(UVariantSet* VarSetOfNode, bool bExpanded);
 	bool GetExpansionState(UVariantSet* VarSetOfNode);
 
+	UFUNCTION(BlueprintCallable, Category="LevelVariantSets")
+	int32 GetNumVariantSets();
+
+	UFUNCTION(BlueprintCallable, Category="LevelVariantSets")
+	UVariantSet* GetVariantSet(int32 VariantSetIndex);
+
 private:
 
 	UPROPERTY()
diff --git a/Engine/Plugins/Enterprise/VariantManagerContent/Source/VariantManagerContent/Public/LevelVariantSetsActor.h b/Engine/Plugins/Enterprise/VariantManagerContent/Source/VariantManagerContent/Public/LevelVariantSetsActor.h
--- a/Engine/Plugins/Enterprise/VariantManagerContent/Source/VariantManagerContent/Public/LevelVariantSetsActor.h
+++ b/Engine/Plugins/Enterprise/VariantManagerContent/Source/VariantManagerContent/Public/LevelVariantSetsActor.h
@@ -11,6 +11,9 @@
 #include "LevelVariantSetsActor.generated.h"
 
 
+class UVariantSet;
+
+
 UCLASS(hideCategories=(Rendering, Physics, LOD, Activation, Input))
 class VARIANTMANAGERCONTENT_API ALevelVariantSetsActor : public AActor
 {
@@ -22,14 +25,20 @@ public:
 
 public:
 
-	UPROPERTY(EditAnywhere, BlueprintReadOnly, Category="Variants", meta=(AllowedClasses="LevelVariantSets"))
+	UPROPERTY(EditAnywhere, BlueprintReadOnly, Category="LevelVariantSets", meta=(AllowedClasses="LevelVariantSets"))
 	FSoftObjectPath LevelVariantSets;
 
 public:
 
-	UFUNCTION(BlueprintCallable, Category="Variants")
+	UFUNCTION(BlueprintCallable, Category="LevelVariantSets")
 	ULevelVariantSets* GetLevelVariantSets(bool bLoad = false) const;
 
-	UFUNCTION(BlueprintCallable, Category="Variants")
+	UFUNCTION(BlueprintCallable, Category="LevelVariantSets")
 	void SetLevelVariantSets(ULevelVariantSets* InVariantSets);
+
+	UFUNCTION(BlueprintCallable, Category="LevelVariantSets")
+	bool SwitchOnVariantByName(FString VariantSetName, FString VariantName);
+
+	UFUNCTION(BlueprintCallable, Category="LevelVariantSets")
+	bool SwitchOnVariantByIndex(int32 VariantSetIndex, int32 VariantIndex);
 };
diff --git a/Engine/Plugins/Enterprise/VariantManagerContent/Source/VariantManagerContent/Public/PropertyValue.h b/Engine/Plugins/Enterprise/VariantManagerContent/Source/VariantManagerContent/Public/PropertyValue.h
new file mode 100644
--- /dev/null
+++ b/Engine/Plugins/Enterprise/VariantManagerContent/Source/VariantManagerContent/Public/PropertyValue.h
@@ -0,0 +1,118 @@
+// Copyright 1998-2018 Epic Games, Inc. All Rights Reserved.
+
+#pragma once
+
+#include "CoreMinimal.h"
+#include "UObject/ObjectMacros.h"
+
+#include "PropertyValue.generated.h"
+
+#define PATH_DELIMITER TEXT(" / ")
+
+VARIANTMANAGERCONTENT_API DECLARE_LOG_CATEGORY_EXTERN(LogVariantContent, Log, All);
+
+
+DECLARE_MULTICAST_DELEGATE(FOnPropertyRecorded);
+DECLARE_MULTICAST_DELEGATE(FOnPropertyApplied);
+
+
+class UVariantObjectBinding;
+
+
+UENUM()
+enum class EPropertyValueCategory : uint8
+{
+	Generic,
+	RelativeLocation,
+	RelativeRotation,
+	RelativeScale3D,
+	bVisible,
+	Material
+};
+
+UCLASS(BlueprintType)
+class VARIANTMANAGERCONTENT_API UPropertyValue : public UObject
+{
+	GENERATED_UCLASS_BODY()
+
+public:
+
+	void Init(const TArray<UProperty*>& InProps, const TArray<int32> InIndices, const FString& InFullDisplayString, EPropertyValueCategory InCategory = EPropertyValueCategory::Generic);
+	const TArray<UProperty*>& GetProperties() const;
+	const TArray<int32>& GetPropertyIndices() const;
+	class UVariantObjectBinding* GetParent() const;
+
+	// UObject Interface
+	virtual void Serialize(FArchive& Ar) override;
+	//~ End UObject Interface
+
+	// Checks if our parent object binding has our property. Once resolved we can use the other utility
+	// functions like get container or property addresses
+	virtual bool Resolve();
+	bool HasValidResolve() const;
+	void ClearLastResolve();
+
+	// Need valid resolve
+	void* GetPropertyParentContainerAddress() const;
+	UStruct* GetPropertyParentContainerClass() const;
+	virtual void RecordDataFromResolvedObject();
+	virtual void ApplyDataToResolvedObject();
+	//~ End valid resolve requirement
+
+	// Returns the type of UProperty (UObjectProperty, UFloatProperty, etc)
+	UClass* GetPropertyClass() const;
+	// If we're a UStructProperty, returns the type of UScriptStruct that we can point at
+	UScriptStruct* GetStructPropertyStruct() const;
+	// If we're an UObjectProperty, returns the UClass of UObjects that we can point at
+	UClass* GetObjectPropertyObjectClass() const;
+	FName GetPropertyName() const;
+	const FString& GetFullDisplayString() const;
+	FString GetLeafDisplayString() const;
+	int32 GetValueSizeInBytes() const;
+	int32 GetPropertyOffsetInBytes() const;
+
+	bool HasRecordedData() const;
+	const TArray<uint8>& GetRecordedData();
+	void SetRecordedData(const uint8* NewDataBytes, int32 NumBytes);
+
+	FOnPropertyApplied& GetOnPropertyApplied();
+	FOnPropertyRecorded& GetOnPropertyRecorded();
+
+protected:
+
+	UProperty* GetProperty() const;
+
+	bool ResolvePropertiesRecursive(UStruct* ContainerClass, void* ContainerAddress, int32 PropertyIndex);
+
+	FOnPropertyApplied OnPropertyApplied;
+	FOnPropertyRecorded OnPropertyRecorded;
+
+	// Temp data cached from last resolve
+	UProperty* LeafProperty;
+	UStruct* ParentContainerClass;
+	void* ParentContainerAddress;
+	uint8* PropertyValuePtr;
+
+	UPROPERTY()
+	FString FullDisplayString;
+
+	UPROPERTY()
+	bool bHasRecordedData;
+
+	UPROPERTY()
+	bool bIsObjectProperty;
+
+	UPROPERTY()
+	TArray<uint8> ValueBytes;
+
+	UPROPERTY()
+	EPropertyValueCategory PropCategory;
+
+	UPROPERTY()
+	TArray<UProperty*> Properties;
+
+	UPROPERTY()
+	TArray<int32> PropertyIndices;
+
+	TSoftObjectPtr<UObject> TempObjPtr;
+};
\ No newline at end of file
diff --git a/Engine/Plugins/Enterprise/VariantManagerContent/Source/VariantManagerContent/Public/PropertyValueMaterial.h b/Engine/Plugins/Enterprise/VariantManagerContent/Source/VariantManagerContent/Public/PropertyValueMaterial.h
new file mode 100644
--- /dev/null
+++ b/Engine/Plugins/Enterprise/VariantManagerContent/Source/VariantManagerContent/Public/PropertyValueMaterial.h
@@ -0,0 +1,28 @@
+// Copyright 1998-2018 Epic Games, Inc. All Rights Reserved.
+
+#pragma once
+
+#include "CoreMinimal.h"
+#include "UObject/ObjectMacros.h"
+#include "PropertyValue.h"
+
+#include "PropertyValueMaterial.generated.h"
+
+
+UCLASS(BlueprintType)
+class VARIANTMANAGERCONTENT_API UPropertyValueMaterial : public UPropertyValue
+{
+	GENERATED_UCLASS_BODY()
+
+public:
+
+	// Our leaf property will always be OverrideMaterials/OverrideMaterials[0] just for the
+	// type/size/class information. It will normally fail to resolve if the StaticMeshComponent
+	// is just using default materials though, so we have to intercept resolve calls and handle
+	// them in a specific way. This will also let us zero out the value ptr and other things
+	// that shouldn't be used by themselves
+	virtual bool Resolve();
+
+	virtual void RecordDataFromResolvedObject() override;
+	virtual void ApplyDataToResolvedObject() override;
+};
\ No newline at end of file
diff --git a/Engine/Plugins/Enterprise/VariantManagerContent/Source/VariantManagerContent/Public/PropertyValueTransform.h b/Engine/Plugins/Enterprise/VariantManagerContent/Source/VariantManagerContent/Public/PropertyValueTransform.h
new file mode 100644
--- /dev/null
+++ b/Engine/Plugins/Enterprise/VariantManagerContent/Source/VariantManagerContent/Public/PropertyValueTransform.h
@@ -0,0 +1,20 @@
+// Copyright 1998-2018 Epic Games, Inc. All Rights Reserved.
+
+#pragma once
+
+#include "CoreMinimal.h"
+#include "UObject/ObjectMacros.h"
+#include "PropertyValue.h"
+
+#include "PropertyValueTransform.generated.h"
+
+
+UCLASS(BlueprintType)
+class VARIANTMANAGERCONTENT_API UPropertyValueTransform : public UPropertyValue
+{
+	GENERATED_UCLASS_BODY()
+
+public:
+
+	virtual void ApplyDataToResolvedObject() override;
+};
\ No newline at end of file
diff --git a/Engine/Plugins/Enterprise/VariantManagerContent/Source/VariantManagerContent/Public/PropertyValueVisibility.h b/Engine/Plugins/Enterprise/VariantManagerContent/Source/VariantManagerContent/Public/PropertyValueVisibility.h
new file mode 100644
--- /dev/null
+++ b/Engine/Plugins/Enterprise/VariantManagerContent/Source/VariantManagerContent/Public/PropertyValueVisibility.h
@@ -0,0 +1,20 @@
+// Copyright 1998-2018 Epic Games, Inc. All Rights Reserved.
+
+#pragma once
+
+#include "CoreMinimal.h"
+#include "UObject/ObjectMacros.h"
+#include "PropertyValue.h"
+
+#include "PropertyValueVisibility.generated.h"
+
+
+UCLASS(BlueprintType)
+class VARIANTMANAGERCONTENT_API UPropertyValueVisibility : public UPropertyValue
+{
+	GENERATED_UCLASS_BODY()
+
+public:
+
+	virtual void ApplyDataToResolvedObject() override;
+};
\ No newline at end of file
diff --git a/Engine/Plugins/Enterprise/VariantManagerContent/Source/VariantManagerContent/Public/Variant.h b/Engine/Plugins/Enterprise/VariantManagerContent/Source/VariantManagerContent/Public/Variant.h
--- a/Engine/Plugins/Enterprise/VariantManagerContent/Source/VariantManagerContent/Public/Variant.h
+++ b/Engine/Plugins/Enterprise/VariantManagerContent/Source/VariantManagerContent/Public/Variant.h
@@ -3,11 +3,14 @@
 #pragma once
 
 #include "CoreMinimal.h"
+#include "Internationalization/Text.h"
+#include "UObject/Object.h"
 #include "UObject/ObjectMacros.h"
+#include "UObject/TextProperty.h"
 
 #include "Variant.generated.h"
 
-
+class AActor;
 class UVariantObjectBinding;
 
 
@@ -18,49 +21,33 @@ class VARIANTMANAGERCONTENT_API UVariant : public UObject
 
 public:
 
-	/**~ UObject implementation */
-	//virtual void Serialize(FArchive& Ar) override;
-
-	// We need this because UVariantObjectBindings have non-UPROPERTY FLazyObjectPtrs
-	// that need to be manually copied
-	UVariant* Clone(UObject* ClonesOuter = INVALID_OBJECT);
-
 	class UVariantSet* GetParent();
 
-	void SetDisplayName(const FText& NewDisplayName);
-	FText GetDisplayName() const;
-	FText GetDefaultDisplayName() const;
+	UFUNCTION(BlueprintCallable, Category="Variant")
+	void SetDisplayText(const FText& NewDisplayText);
 
-	void AddActors(TWeakObjectPtr<AActor> InActor);
-	void AddActors(const TArray<TWeakObjectPtr<AActor>>& InActors);
+	UFUNCTION(BlueprintCallable, Category="Variant")
+	FText GetDisplayText() const;
 
-	void AddBinding(UVariantObjectBinding* NewBinding);
+	// In case of a duplicate binding these will destroy the older bindings
+	void AddBindings(const TArray<UVariantObjectBinding*>& NewBindings, int32 Index = INDEX_NONE);
+	const TArray<UVariantObjectBinding*>& GetBindings();
+	void RemoveBindings(const TArray<UVariantObjectBinding*>& Bindings);
 
-	const TArray<UVariantObjectBinding*>& GetBindings()
-	{
-		return ObjectBindings;
-	}
+	UFUNCTION(BlueprintCallable, Category="Variant")
+	int32 GetNumActors();
 
-	void RemoveBinding(UVariantObjectBinding* ThisBinding);
+	UFUNCTION(BlueprintCallable, Category="Variant")
+	AActor* GetActor(int32 ActorIndex);
 
-	int32 GetSortingOrder() const
-	{
-		return SortingOrder;
-	}
-
-	void SetSortingOrder(const int32 InSortingOrder)
-	{
-		SortingOrder = InSortingOrder;
-	}
+	UFUNCTION(BlueprintCallable, Category="Variant")
+	void SwitchOn();
 
 private:
-	UPROPERTY()
-	FText DisplayName;
 
 	UPROPERTY()
-	int32 SortingOrder;
+	FText DisplayText;
 
-	// We manually duplicate these within our Clone function
-	UPROPERTY(DuplicateTransient)
+	UPROPERTY()
 	TArray<UVariantObjectBinding*> ObjectBindings;
 };
diff --git a/Engine/Plugins/Enterprise/VariantManagerContent/Source/VariantManagerContent/Public/VariantObjectBinding.h b/Engine/Plugins/Enterprise/VariantManagerContent/Source/VariantManagerContent/Public/VariantObjectBinding.h
--- a/Engine/Plugins/Enterprise/VariantManagerContent/Source/VariantManagerContent/Public/VariantObjectBinding.h
+++ b/Engine/Plugins/Enterprise/VariantManagerContent/Source/VariantManagerContent/Public/VariantObjectBinding.h
@@ -4,53 +4,42 @@
 
 #include "CoreMinimal.h"
 #include "UObject/ObjectMacros.h"
+#include "UObject/SoftObjectPath.h"
 #include "UObject/LazyObjectPtr.h"
-#include "Misc/Guid.h"
 
 #include "VariantObjectBinding.generated.h"
 
 
+class UPropertyValue;
+
+
 UCLASS(BlueprintType)
 class VARIANTMANAGERCONTENT_API UVariantObjectBinding : public UObject
 {
 	GENERATED_UCLASS_BODY()
 
 public:
-	void Init(const UObject* InObject, int32 InSortingOrder = -1)
-	{
-		ObjectPtr = FLazyObjectPtr(InObject);
-		SetSortingOrder(InSortingOrder);
-	}
-
-	// We need this because UVariantObjectBindings have non-UPROPERTY FLazyObjectPtrs
-	// that need to be manually copied
-	UVariantObjectBinding* Clone(UObject* ClonesOuter = INVALID_OBJECT);
 
+	void Init(UObject* InObject);
 	class UVariant* GetParent();
 
-	const FGuid& GetObjectGuid() const
-	{
-		return ObjectPtr.GetUniqueID().GetGuid();
-	}
-
-    UObject* GetObject() const
-    {
-        return ObjectPtr.Get();
-    }
+	FText GetDisplayText() const;
 
-	int32 GetSortingOrder() const
-	{
-		return SortingOrder;
-	}
+	FString GetObjectPath() const;
+	UObject* GetObject() const;
 
-	void SetSortingOrder(const int32 InSortingOrder)
-	{
-		SortingOrder = InSortingOrder;
-	}
+	void AddCapturedProperties(const TArray<UPropertyValue*>& Properties, int32 Index = INDEX_NONE);
+	const TArray<UPropertyValue*>& GetCapturedProperties() const;
+	void RemoveCapturedProperties(const TArray<UPropertyValue*>& Properties);
 
 private:
-	FLazyObjectPtr ObjectPtr;
 
 	UPROPERTY()
-	int32 SortingOrder;
+	mutable FSoftObjectPath ObjectPtr;
+
+	UPROPERTY()
+	TLazyObjectPtr<UObject> LazyObjectPtr;
+
+	UPROPERTY()
+	TArray<UPropertyValue*> CapturedProperties;
 };
diff --git a/Engine/Plugins/Enterprise/VariantManagerContent/Source/VariantManagerContent/Public/VariantSet.h b/Engine/Plugins/Enterprise/VariantManagerContent/Source/VariantManagerContent/Public/VariantSet.h
--- a/Engine/Plugins/Enterprise/VariantManagerContent/Source/VariantManagerContent/Public/VariantSet.h
+++ b/Engine/Plugins/Enterprise/VariantManagerContent/Source/VariantManagerContent/Public/VariantSet.h
@@ -4,9 +4,11 @@
 
 #include "CoreMinimal.h"
 #include "UObject/ObjectMacros.h"
+#include "UObject/TextProperty.h"
 
 #include "VariantSet.generated.h"
 
+
 class UVariant;
 
 
@@ -17,45 +19,31 @@ class VARIANTMANAGERCONTENT_API UVariantSet : public UObject
 
 public:
 
-	/**~ UObject implementation */
-	//virtual void Serialize(FArchive& Ar) override;
-
-	// We need this because UVariantObjectBindings have non-UPROPERTY FLazyObjectPtrs
-	// that need to be manually copied
-	UVariantSet* Clone(UObject* ClonesOuter = INVALID_OBJECT);
-
 	class ULevelVariantSets* GetParent();
 
-	void SetDisplayName(const FText& NewDisplayName);
-	FText GetDisplayName() const;
-	FText GetDefaultDisplayName() const;
+	UFUNCTION(BlueprintCallable, Category="VariantSet")
+	void SetDisplayText(const FText& NewDisplayText);
 
-	void AddVariant(UVariant* NewVariant);
-	const TArray<UVariant*>& GetVariants() const
-	{
-		return Variants;
-	}
+	UFUNCTION(BlueprintCallable, Category="VariantSet")
+	FText GetDisplayText() const;
 
-	void RemoveVariant(UVariant* ThisVariant);
+	FString GetUniqueVariantName(const FString& InPrefix);
 
-	int32 GetSortingOrder() const
-	{
-		return SortingOrder;
-	}
+	void AddVariants(const TArray<UVariant*>& NewVariants, int32 Index = INDEX_NONE);
+	const TArray<UVariant*>& GetVariants() const;
+	void RemoveVariants(const TArray<UVariant*>& InVariants);
 
-	void SetSortingOrder(const int32 InSortingOrder)
-	{
-		SortingOrder = InSortingOrder;
-	}
+	UFUNCTION(BlueprintCallable, Category="VariantSet")
+	int32 GetNumVariants();
+
+	UFUNCTION(BlueprintCallable, Category="VariantSet")
+	UVariant* GetVariant(int32 VariantIndex);
 
 private:
-	UPROPERTY()
-	FText DisplayName;
 
 	UPROPERTY()
-	int32 SortingOrder;
+	FText DisplayText;
 
-	// We manually duplicate these within our Clone function
-	UPROPERTY(DuplicateTransient)
+	UPROPERTY()
 	TArray<UVariant*> Variants;
 };
diff --git a/Engine/Plugins/Enterprise/VariantManagerContent/Source/VariantManagerContent/VariantManagerContent.Build.cs b/Engine/Plugins/Enterprise/VariantManagerContent/Source/VariantManagerContent/VariantManagerContent.Build.cs
--- a/Engine/Plugins/Enterprise/VariantManagerContent/Source/VariantManagerContent/VariantManagerContent.Build.cs
+++ b/Engine/Plugins/Enterprise/VariantManagerContent/Source/VariantManagerContent/VariantManagerContent.Build.cs
@@ -12,14 +12,7 @@ namespace UnrealBuildTool.Rules
 				{
 					"Core",
 					"CoreUObject",
-					"Engine",
-					"MovieScene"
-				}
-			);
-
-			PrivateDependencyModuleNames.AddRange(
-				new string[]
-				{
+					"Engine"
 				}
 			);
 		}
diff --git a/Engine/Source/Runtime/Engine/Classes/Materials/MaterialExpressionArccosine.h b/Engine/Source/Runtime/Engine/Classes/Materials/MaterialExpressionArccosine.h
--- a/Engine/Source/Runtime/Engine/Classes/Materials/MaterialExpressionArccosine.h
+++ b/Engine/Source/Runtime/Engine/Classes/Materials/MaterialExpressionArccosine.h
@@ -8,7 +8,7 @@
 #include "Materials/MaterialExpression.h"
 #include "MaterialExpressionArccosine.generated.h"
 
-UCLASS(collapsecategories, hidecategories=Object)
+UCLASS(MinimalAPI, collapsecategories, hidecategories=Object)
 class UMaterialExpressionArccosine : public UMaterialExpression
 {
 	GENERATED_UCLASS_BODY()
diff --git a/Engine/Source/Runtime/Engine/Classes/Materials/MaterialExpressionArcsine.h b/Engine/Source/Runtime/Engine/Classes/Materials/MaterialExpressionArcsine.h
--- a/Engine/Source/Runtime/Engine/Classes/Materials/MaterialExpressionArcsine.h
+++ b/Engine/Source/Runtime/Engine/Classes/Materials/MaterialExpressionArcsine.h
@@ -8,7 +8,7 @@
 #include "Materials/MaterialExpression.h"
 #include "MaterialExpressionArcsine.generated.h"
 
-UCLASS(collapsecategories, hidecategories=Object)
+UCLASS(MinimalAPI, collapsecategories, hidecategories=Object)
 class UMaterialExpressionArcsine : public UMaterialExpression
 {
 	GENERATED_UCLASS_BODY()
diff --git a/Engine/Source/Runtime/Engine/Classes/Materials/MaterialExpressionArctangent.h b/Engine/Source/Runtime/Engine/Classes/Materials/MaterialExpressionArctangent.h
--- a/Engine/Source/Runtime/Engine/Classes/Materials/MaterialExpressionArctangent.h
+++ b/Engine/Source/Runtime/Engine/Classes/Materials/MaterialExpressionArctangent.h
@@ -8,7 +8,7 @@
 #include "Materials/MaterialExpression.h"
 #include "MaterialExpressionArctangent.generated.h"
 
-UCLASS(collapsecategories, hidecategories=Object)
+UCLASS(MinimalAPI, collapsecategories, hidecategories=Object)
 class UMaterialExpressionArctangent : public UMaterialExpression
 {
 	GENERATED_UCLASS_BODY()
diff --git a/Engine/Source/Runtime/Engine/Classes/Materials/MaterialExpressionArctangent2.h b/Engine/Source/Runtime/Engine/Classes/Materials/MaterialExpressionArctangent2.h
--- a/Engine/Source/Runtime/Engine/Classes/Materials/MaterialExpressionArctangent2.h
+++ b/Engine/Source/Runtime/Engine/Classes/Materials/MaterialExpressionArctangent2.h
@@ -8,7 +8,7 @@
 #include "Materials/MaterialExpression.h"
 #include "MaterialExpressionArctangent2.generated.h"
 
-UCLASS(collapsecategories, hidecategories=Object)
+UCLASS(MinimalAPI, collapsecategories, hidecategories=Object)
 class UMaterialExpressionArctangent2 : public UMaterialExpression
 {
 	GENERATED_UCLASS_BODY()
diff --git a/Engine/Source/Runtime/Engine/Classes/Materials/MaterialExpressionBlackBody.h b/Engine/Source/Runtime/Engine/Classes/Materials/MaterialExpressionBlackBody.h
--- a/Engine/Source/Runtime/Engine/Classes/Materials/MaterialExpressionBlackBody.h
+++ b/Engine/Source/Runtime/Engine/Classes/Materials/MaterialExpressionBlackBody.h
@@ -9,7 +9,7 @@
 #include "Materials/MaterialExpression.h"
 #include "MaterialExpressionBlackBody.generated.h"
 
-UCLASS()
+UCLASS(MinimalAPI)
 class UMaterialExpressionBlackBody : public UMaterialExpression
 {
 	GENERATED_UCLASS_BODY()
diff --git a/Engine/Source/Runtime/Engine/Classes/Materials/MaterialExpressionCeil.h b/Engine/Source/Runtime/Engine/Classes/Materials/MaterialExpressionCeil.h
--- a/Engine/Source/Runtime/Engine/Classes/Materials/MaterialExpressionCeil.h
+++ b/Engine/Source/Runtime/Engine/Classes/Materials/MaterialExpressionCeil.h
@@ -9,7 +9,7 @@
 #include "Materials/MaterialExpression.h"
 #include "MaterialExpressionCeil.generated.h"
 
-UCLASS(collapsecategories, hidecategories=Object)
+UCLASS(MinimalAPI, collapsecategories, hidecategories=Object)
 class UMaterialExpressionCeil : public UMaterialExpression
 {
 	GENERATED_UCLASS_BODY()
diff --git a/Engine/Source/Runtime/Engine/Classes/Materials/MaterialExpressionClearCoatNormalCustomOutput.h b/Engine/Source/Runtime/Engine/Classes/Materials/MaterialExpressionClearCoatNormalCustomOutput.h
--- a/Engine/Source/Runtime/Engine/Classes/Materials/MaterialExpressionClearCoatNormalCustomOutput.h
+++ b/Engine/Source/Runtime/Engine/Classes/Materials/MaterialExpressionClearCoatNormalCustomOutput.h
@@ -9,7 +9,7 @@
 #include "Materials/MaterialExpressionCustomOutput.h"
 #include "MaterialExpressionClearCoatNormalCustomOutput.generated.h"
 
-UCLASS()
+UCLASS(MinimalAPI)
 class UMaterialExpressionClearCoatNormalCustomOutput : public UMaterialExpressionCustomOutput
 {
 	GENERATED_UCLASS_BODY()
diff --git a/Engine/Source/Runtime/Engine/Classes/Materials/MaterialExpressionDistance.h b/Engine/Source/Runtime/Engine/Classes/Materials/MaterialExpressionDistance.h
--- a/Engine/Source/Runtime/Engine/Classes/Materials/MaterialExpressionDistance.h
+++ b/Engine/Source/Runtime/Engine/Classes/Materials/MaterialExpressionDistance.h
@@ -9,7 +9,7 @@
 #include "Materials/MaterialExpression.h"
 #include "MaterialExpressionDistance.generated.h"
 
-UCLASS()
+UCLASS(MinimalAPI)
 class UMaterialExpressionDistance : public UMaterialExpression
 {
 	GENERATED_UCLASS_BODY()
diff --git a/Engine/Source/Runtime/Engine/Classes/Materials/MaterialExpressionFrac.h b/Engine/Source/Runtime/Engine/Classes/Materials/MaterialExpressionFrac.h
--- a/Engine/Source/Runtime/Engine/Classes/Materials/MaterialExpressionFrac.h
+++ b/Engine/Source/Runtime/Engine/Classes/Materials/MaterialExpressionFrac.h
@@ -9,7 +9,7 @@
 #include "Materials/MaterialExpression.h"
 #include "MaterialExpressionFrac.generated.h"
 
-UCLASS(collapsecategories, hidecategories=Object)
+UCLASS(MinimalAPI, collapsecategories, hidecategories=Object)
 class UMaterialExpressionFrac : public UMaterialExpression
 {
 	GENERATED_UCLASS_BODY()
diff --git a/Engine/Source/Runtime/Engine/Classes/Materials/MaterialExpressionSaturate.h b/Engine/Source/Runtime/Engine/Classes/Materials/MaterialExpressionSaturate.h
--- a/Engine/Source/Runtime/Engine/Classes/Materials/MaterialExpressionSaturate.h
+++ b/Engine/Source/Runtime/Engine/Classes/Materials/MaterialExpressionSaturate.h
@@ -8,7 +8,7 @@
 #include "Materials/MaterialExpression.h"
 #include "MaterialExpressionSaturate.generated.h"
 
-UCLASS(collapsecategories, hidecategories=Object)
+UCLASS(MinimalAPI, collapsecategories, hidecategories=Object)
 class UMaterialExpressionSaturate : public UMaterialExpression
 {
 	GENERATED_UCLASS_BODY()
diff --git a/Engine/Source/Runtime/Engine/Classes/Materials/MaterialExpressionSphereMask.h b/Engine/Source/Runtime/Engine/Classes/Materials/MaterialExpressionSphereMask.h
--- a/Engine/Source/Runtime/Engine/Classes/Materials/MaterialExpressionSphereMask.h
+++ b/Engine/Source/Runtime/Engine/Classes/Materials/MaterialExpressionSphereMask.h
@@ -9,7 +9,7 @@
 #include "Materials/MaterialExpression.h"
 #include "MaterialExpressionSphereMask.generated.h"
 
-UCLASS()
+UCLASS(MinimalAPI)
 class UMaterialExpressionSphereMask : public UMaterialExpression
 {
 	GENERATED_UCLASS_BODY()
diff --git a/Engine/Source/Runtime/Engine/Classes/Materials/MaterialExpressionTangent.h b/Engine/Source/Runtime/Engine/Classes/Materials/MaterialExpressionTangent.h
--- a/Engine/Source/Runtime/Engine/Classes/Materials/MaterialExpressionTangent.h
+++ b/Engine/Source/Runtime/Engine/Classes/Materials/MaterialExpressionTangent.h
@@ -8,7 +8,7 @@
 #include "Materials/MaterialExpression.h"
 #include "MaterialExpressionTangent.generated.h"
 
-UCLASS(collapsecategories, hidecategories=Object)
+UCLASS(MinimalAPI, collapsecategories, hidecategories=Object)
 class UMaterialExpressionTangent : public UMaterialExpression
 {
 	GENERATED_UCLASS_BODY()
diff --git a/Engine/Source/Runtime/Engine/Classes/Materials/MaterialExpressionTime.h b/Engine/Source/Runtime/Engine/Classes/Materials/MaterialExpressionTime.h
--- a/Engine/Source/Runtime/Engine/Classes/Materials/MaterialExpressionTime.h
+++ b/Engine/Source/Runtime/Engine/Classes/Materials/MaterialExpressionTime.h
@@ -8,7 +8,7 @@
 #include "Materials/MaterialExpression.h"
 #include "MaterialExpressionTime.generated.h"
 
-UCLASS(collapsecategories, hidecategories=Object)
+UCLASS(MinimalAPI, collapsecategories, hidecategories=Object)
 class UMaterialExpressionTime : public UMaterialExpression
 {
 	GENERATED_UCLASS_BODY()
