Copied 4405429 JeanMichel.Dignard
Build fix - variable shadowing warning in variant manager content

Copied 4410154 JeanMichel.Dignard
Build fix
- Fixed wrong initialization order in VariantManagerPropertyNode (Mac editor)

diff --git a/Engine/Plugins/Enterprise/VariantManagerContent/Source/VariantManagerContent/Private/PropertyValue.cpp b/Engine/Plugins/Enterprise/VariantManagerContent/Source/VariantManagerContent/Private/PropertyValue.cpp
--- a/Engine/Plugins/Enterprise/VariantManagerContent/Source/VariantManagerContent/Private/PropertyValue.cpp
+++ b/Engine/Plugins/Enterprise/VariantManagerContent/Source/VariantManagerContent/Private/PropertyValue.cpp
@@ -465,10 +465,10 @@ bool UPropertyValue::ResolvePropertiesRecursive(UStruct* ContainerClass, void* C
 					// The next link in the chain is just this array's inner. Let's just skip it instead
 					return ResolvePropertiesRecursive(ArrayOfStructsProp->Struct, StructAddress, SegmentIndex + 2);
 				}
-				if ( UObjectProperty* ObjectProperty = Cast<UObjectProperty>(ArrayProp->Inner) )
+				if ( UObjectProperty* InnerObjectProperty = Cast<UObjectProperty>(ArrayProp->Inner) )
 				{
 					void* ObjPtrContainer = ArrayHelper.GetRawPtr(InnerArrayIndex);
-					UObject* CurrentObject = ObjectProperty->GetObjectPropertyValue(ObjPtrContainer);
+					UObject* CurrentObject = InnerObjectProperty->GetObjectPropertyValue(ObjPtrContainer);
 					if (CurrentObject)
 					{
 						ParentContainerClass = CurrentObject->GetClass();
