Changed the event name for MDL analytics and fixed the duration attribute.
Made VariantManager a beta plugin.

diff --git a/Engine/Plugins/Enterprise/VariantManagerContent/VariantManagerContent.uplugin b/Engine/Plugins/Enterprise/VariantManagerContent/VariantManagerContent.uplugin
--- a/Engine/Plugins/Enterprise/VariantManagerContent/VariantManagerContent.uplugin
+++ b/Engine/Plugins/Enterprise/VariantManagerContent/VariantManagerContent.uplugin
@@ -12,7 +12,7 @@
 	"SupportURL": "",
 	"EnabledByDefault" : true,
 	"CanContainContent": true,
-	"IsBetaVersion": false,
+	"IsBetaVersion": true,
 	"Installed": false,
 	"Modules" :
 	[
