Copied 4442004 Daniel.Coelho

Fix for being unable to trigger on variants from UMG blueprints

diff --git a/Engine/Plugins/Enterprise/VariantManagerContent/Source/VariantManagerContent/Private/LevelVariantSetsActor.cpp b/Engine/Plugins/Enterprise/VariantManagerContent/Source/VariantManagerContent/Private/LevelVariantSetsActor.cpp
--- a/Engine/Plugins/Enterprise/VariantManagerContent/Source/VariantManagerContent/Private/LevelVariantSetsActor.cpp
+++ b/Engine/Plugins/Enterprise/VariantManagerContent/Source/VariantManagerContent/Private/LevelVariantSetsActor.cpp
@@ -5,6 +5,7 @@
 #include "LevelVariantSets.h"
 #include "VariantSet.h"
 #include "Variant.h"
+#include "VariantObjectBinding.h"
 
 #include "CoreMinimal.h"
 #include "UObject/ObjectMacros.h"
@@ -19,6 +20,25 @@ ALevelVariantSetsActor::ALevelVariantSetsActor(const FObjectInitializer& Init)
 	RootComponent = SceneComponent;
 }
 
+void ALevelVariantSetsActor::BeginPlay()
+{
+	// See comment on UVariantObjectBinding->FixupForPIE for an explanation as to why we need to do this
+	ULevelVariantSets* LevelVarSets = GetLevelVariantSets(true);
+	if (LevelVarSets)
+	{
+		for (UVariantSet* VarSet : LevelVarSets->GetVariantSets())
+		{
+			for (UVariant* Var : VarSet->GetVariants())
+			{
+				for (UVariantObjectBinding* Binding : Var->GetBindings())
+				{
+					Binding->FixupForPIE();
+				}
+			}
+		}
+	}
+}
+
 ULevelVariantSets* ALevelVariantSetsActor::GetLevelVariantSets(bool bLoad) const
 {
 	if (LevelVariantSets.IsValid())
diff --git a/Engine/Plugins/Enterprise/VariantManagerContent/Source/VariantManagerContent/Private/VariantObjectBinding.cpp b/Engine/Plugins/Enterprise/VariantManagerContent/Source/VariantManagerContent/Private/VariantObjectBinding.cpp
--- a/Engine/Plugins/Enterprise/VariantManagerContent/Source/VariantManagerContent/Private/VariantObjectBinding.cpp
+++ b/Engine/Plugins/Enterprise/VariantManagerContent/Source/VariantManagerContent/Private/VariantObjectBinding.cpp
@@ -3,6 +3,7 @@
 #include "VariantObjectBinding.h"
 
 #include "UObject/LazyObjectPtr.h"
+#include "UObject/SoftObjectPath.h"
 #include "PropertyValue.h"
 #include "Variant.h"
 
@@ -70,6 +71,19 @@ UObject* UVariantObjectBinding::GetObject() const
 	return nullptr;
 }
 
+void UVariantObjectBinding::FixupForPIE()
+{
+	// For some reason calls from UMG blueprints hit FSoftObjectPaths while GPlayInEditorID is INDEX_NONE,
+	// so all object bindings would target the actors back in the editor world.
+	// Calls from the level blueprint or actor blueprints work fine, as they call FSoftObjectPath.ResolveObject()
+	// in a moment where the GPlayInEditorID is a valid index so that ResolveObject can internally call it's FixupForPIE()
+	// and update the path.
+	// We expose this here so that the ALevelVariantSets actor can call this on BeginPlay and update our path for PIE,
+	// as that is also an instant where GPlayInEditorID is valid.
+
+	ObjectPtr.FixupForPIE();
+}
+
 void UVariantObjectBinding::AddCapturedProperties(const TArray<UPropertyValue*>& NewProperties, int32 Index)
 {
 	Modify();
diff --git a/Engine/Plugins/Enterprise/VariantManagerContent/Source/VariantManagerContent/Public/LevelVariantSetsActor.h b/Engine/Plugins/Enterprise/VariantManagerContent/Source/VariantManagerContent/Public/LevelVariantSetsActor.h
--- a/Engine/Plugins/Enterprise/VariantManagerContent/Source/VariantManagerContent/Public/LevelVariantSetsActor.h
+++ b/Engine/Plugins/Enterprise/VariantManagerContent/Source/VariantManagerContent/Public/LevelVariantSetsActor.h
@@ -23,12 +23,9 @@ public:
 
 	ALevelVariantSetsActor(const FObjectInitializer& Init);
 
-public:
-
-	UPROPERTY(EditAnywhere, BlueprintReadOnly, Category="LevelVariantSets", meta=(AllowedClasses="LevelVariantSets"))
-	FSoftObjectPath LevelVariantSets;
-
-public:
+	// AActor interface
+	virtual void BeginPlay() override;
+	//~ End AActor interface
 
 	UFUNCTION(BlueprintCallable, Category="LevelVariantSets")
 	ULevelVariantSets* GetLevelVariantSets(bool bLoad = false) const;
@@ -41,4 +38,9 @@ public:
 
 	UFUNCTION(BlueprintCallable, Category="LevelVariantSets")
 	bool SwitchOnVariantByIndex(int32 VariantSetIndex, int32 VariantIndex);
+
+public:
+
+	UPROPERTY(EditAnywhere, BlueprintReadOnly, Category="LevelVariantSets", meta=(AllowedClasses="LevelVariantSets"))
+	FSoftObjectPath LevelVariantSets;
 };
diff --git a/Engine/Plugins/Enterprise/VariantManagerContent/Source/VariantManagerContent/Public/VariantObjectBinding.h b/Engine/Plugins/Enterprise/VariantManagerContent/Source/VariantManagerContent/Public/VariantObjectBinding.h
--- a/Engine/Plugins/Enterprise/VariantManagerContent/Source/VariantManagerContent/Public/VariantObjectBinding.h
+++ b/Engine/Plugins/Enterprise/VariantManagerContent/Source/VariantManagerContent/Public/VariantObjectBinding.h
@@ -28,6 +28,11 @@ public:
 	FString GetObjectPath() const;
 	UObject* GetObject() const;
 
+	// Calls FixupForPIE on our ObjectPtr. It is necessary to call this from a moment where
+	// GPlayInEditorID has a valid ID (like an AActor's BeginPlay), as UMG blueprint function calls
+	// are not one of those. See the comment in the implementation of this function for more details
+	void FixupForPIE();
+
 	void AddCapturedProperties(const TArray<UPropertyValue*>& Properties, int32 Index = INDEX_NONE);
 	const TArray<UPropertyValue*>& GetCapturedProperties() const;
 	void RemoveCapturedProperties(const TArray<UPropertyValue*>& Properties);
