Copied cl 4497016 Daniel.Coelho

VariantManager
- Fix recursive asset loading crash

diff --git a/Engine/Plugins/Enterprise/VariantManagerContent/Source/VariantManagerContent/Private/PropertyValue.cpp b/Engine/Plugins/Enterprise/VariantManagerContent/Source/VariantManagerContent/Private/PropertyValue.cpp
--- a/Engine/Plugins/Enterprise/VariantManagerContent/Source/VariantManagerContent/Private/PropertyValue.cpp
+++ b/Engine/Plugins/Enterprise/VariantManagerContent/Source/VariantManagerContent/Private/PropertyValue.cpp
@@ -53,44 +53,32 @@ void UPropertyValue::Serialize(FArchive& Ar)
 
 	if (Ar.IsSaving())
 	{
-		// Create a soft object ptr with our UObject* data bytes so that it can persist
-		if (bIsObjectProperty && HasRecordedData())
+		// If the pointer is not null, means we haven't dealt with it yet (haven't needed GetRecordedData)
+		// so just save it back the way we received it
+		// If our pointer is null but we have bytes representing an UObject reference, it means we
+		// read our pointer at some point (or just created this propertyvalue), so we need to create it
+		if (TempObjPtr.IsNull())
 		{
-			UObject* Obj = *((UObject**)ValueBytes.GetData());
-			if (Obj && Obj->IsValidLowLevel())
+			if (bIsObjectProperty && HasRecordedData())
 			{
-				TempObjPtr = Obj;
+				UObject* Obj = *((UObject**)ValueBytes.GetData());
+				if (Obj && Obj->IsValidLowLevel())
+				{
+					TempObjPtr = Obj;
+				}
+			}
+			else
+			{
+				TempObjPtr.Reset();
 			}
 		}
-		else
-		{
-			TempObjPtr.Reset();
-		}
+
 		Ar << TempObjPtr;
 	}
 	else if (Ar.IsLoading())
 	{
 		Ar << TempObjPtr;
 
-		if (bIsObjectProperty)
-		{
-			// Force resolve of our soft object pointer
-			//@Speed Maybe delay this and let it load whenever it wants?
-			UObject* Obj = TempObjPtr.LoadSynchronous();
-
-			if (Obj && Obj->IsValidLowLevel())
-			{
-				int32 NumBytes = sizeof(UObject**);
-				ValueBytes.SetNumUninitialized(NumBytes);
-				FMemory::Memcpy(ValueBytes.GetData(), &Obj, NumBytes);
-				bHasRecordedData = true;
-			}
-			else
-			{
-				TempObjPtr.Reset();
-				bHasRecordedData = false;
-			}
-		}
 	}
 }
 
@@ -315,9 +303,32 @@ bool UPropertyValue::HasRecordedData() const
 
 const TArray<uint8>& UPropertyValue::GetRecordedData()
 {
-	check(bHasRecordedData)
+	check(bHasRecordedData);
 
 	ValueBytes.SetNum(GetValueSizeInBytes());
+
+	// We need to resolve our softpath still
+	if (bHasRecordedData && bIsObjectProperty && !TempObjPtr.IsNull())
+	{
+		// Force resolve of our soft object pointer
+		UObject* Obj = TempObjPtr.LoadSynchronous();
+
+		if (Obj && Obj->IsValidLowLevel())
+		{
+			int32 NumBytes = sizeof(UObject**);
+			ValueBytes.SetNumUninitialized(NumBytes);
+			FMemory::Memcpy(ValueBytes.GetData(), &Obj, NumBytes);
+
+			bHasRecordedData = true;
+		}
+		else
+		{
+			bHasRecordedData = false;
+		}
+
+		TempObjPtr.Reset();
+	}
+
 	return ValueBytes;
 }
 
diff --git a/Engine/Plugins/Enterprise/VariantManagerContent/Source/VariantManagerContent/Private/PropertyValueMaterial.cpp b/Engine/Plugins/Enterprise/VariantManagerContent/Source/VariantManagerContent/Private/PropertyValueMaterial.cpp
--- a/Engine/Plugins/Enterprise/VariantManagerContent/Source/VariantManagerContent/Private/PropertyValueMaterial.cpp
+++ b/Engine/Plugins/Enterprise/VariantManagerContent/Source/VariantManagerContent/Private/PropertyValueMaterial.cpp
@@ -125,7 +125,7 @@ void UPropertyValueMaterial::ApplyDataToResolvedObject()
 	ContainerObject->SetFlags(RF_Transactional);
 	ContainerObject->Modify();
 
-	UMaterialInterface* Mat = *((UMaterialInterface**)ValueBytes.GetData());
+	UMaterialInterface* Mat = *((UMaterialInterface**)GetRecordedData().GetData());
 
 	int32 NumIndices = PropertyIndices.Num();
 	if (Mat && Mat->IsValidLowLevel() && NumIndices > 0)
