Fixed alignment warning on Win32

diff --git a/Engine/Plugins/Enterprise/VariantManagerContent/Source/VariantManagerContent/Private/PropertyValueTransform.cpp b/Engine/Plugins/Enterprise/VariantManagerContent/Source/VariantManagerContent/Private/PropertyValueTransform.cpp
--- a/Engine/Plugins/Enterprise/VariantManagerContent/Source/VariantManagerContent/Private/PropertyValueTransform.cpp
+++ b/Engine/Plugins/Enterprise/VariantManagerContent/Source/VariantManagerContent/Private/PropertyValueTransform.cpp
@@ -54,7 +54,7 @@ FVector UPropertyValueTransform::GetScale3D()
 	}
 }
 
-void UPropertyValueTransform::SetLocation(FVector NewValue)
+void UPropertyValueTransform::SetLocation(const FVector& NewValue)
 {
 	if (PropCategory == EPropertyValueCategory::RelativeLocation)
 	{
@@ -62,7 +62,7 @@ void UPropertyValueTransform::SetLocation(FVector NewValue)
 	}
 }
 
-void UPropertyValueTransform::SetRotation(FQuat NewValue)
+void UPropertyValueTransform::SetRotation(const FQuat& NewValue)
 {
 	if (PropCategory == EPropertyValueCategory::RelativeRotation)
 	{
@@ -71,7 +71,7 @@ void UPropertyValueTransform::SetRotation(FQuat NewValue)
 	}
 }
 
-void UPropertyValueTransform::SetScale3D(FVector NewValue)
+void UPropertyValueTransform::SetScale3D(const FVector& NewValue)
 {
 	if (PropCategory == EPropertyValueCategory::RelativeScale3D)
 	{
diff --git a/Engine/Plugins/Enterprise/VariantManagerContent/Source/VariantManagerContent/Public/PropertyValueTransform.h b/Engine/Plugins/Enterprise/VariantManagerContent/Source/VariantManagerContent/Public/PropertyValueTransform.h
--- a/Engine/Plugins/Enterprise/VariantManagerContent/Source/VariantManagerContent/Public/PropertyValueTransform.h
+++ b/Engine/Plugins/Enterprise/VariantManagerContent/Source/VariantManagerContent/Public/PropertyValueTransform.h
@@ -19,9 +19,9 @@ public:
 	FQuat GetRotation();
 	FVector GetScale3D();
 
-	void SetLocation(FVector NewValue);
-	void SetRotation(FQuat NewValue);
-	void SetScale3D(FVector NewValue);
+	void SetLocation(const FVector& NewValue);
+	void SetRotation(const FQuat& NewValue);
+	void SetScale3D(const FVector& NewValue);
 
 	virtual void ApplyDataToResolvedObject() override;
 };
\ No newline at end of file
diff --git a/Engine/Source/Runtime/Engine/Classes/Kismet/KismetMathLibrary.h b/Engine/Source/Runtime/Engine/Classes/Kismet/KismetMathLibrary.h
--- a/Engine/Source/Runtime/Engine/Classes/Kismet/KismetMathLibrary.h
+++ b/Engine/Source/Runtime/Engine/Classes/Kismet/KismetMathLibrary.h
@@ -1835,7 +1835,7 @@ class ENGINE_API UKismetMathLibrary : public UBlueprintFunctionLibrary
 
 	/** Breaks a 4D vector apart into X, Y, Z, W. */
 	UFUNCTION(BlueprintPure, meta = (NativeBreakFunc), Category = "Math|Vector4")
-	static void BreakVector4(FVector4 InVec, float& X, float& Y, float& Z, float& W);
+	static void BreakVector4(const FVector4& InVec, float& X, float& Y, float& Z, float& W);
 
 	/** Convert a Vector4 to a Vector (dropping the W element) */
 	UFUNCTION(BlueprintPure, meta = (DisplayName = "To Vector (Vector4)", CompactNodeTitle = "->", ScriptMethod = "Vector", Keywords = "cast convert", BlueprintAutocast), Category = "Math|Conversions")
diff --git a/Engine/Source/Runtime/Engine/Classes/Kismet/KismetMathLibrary.inl b/Engine/Source/Runtime/Engine/Classes/Kismet/KismetMathLibrary.inl
--- a/Engine/Source/Runtime/Engine/Classes/Kismet/KismetMathLibrary.inl
+++ b/Engine/Source/Runtime/Engine/Classes/Kismet/KismetMathLibrary.inl
@@ -1833,7 +1833,7 @@ FVector4 UKismetMathLibrary::MakeVector4(float X, float Y, float Z, float W)
 }
 
 KISMET_MATH_FORCEINLINE
-void UKismetMathLibrary::BreakVector4(FVector4 InVec, float& X, float& Y, float& Z, float& W)
+void UKismetMathLibrary::BreakVector4(const FVector4& InVec, float& X, float& Y, float& Z, float& W)
 {
 	X = InVec.X;
 	Y = InVec.Y;
