Copied 4.22 fixes for VRED, Deltagen and the Variants Manager
cl 4784047, 4784048, 4784049, 4796378, 4797164, 4833865, 4845160, 4846167, 4883014, 4885323, 4890031, 4902933

diff --git a/Engine/Plugins/Enterprise/VariantManagerContent/Source/VariantManagerContent/Private/LevelVariantSets.cpp b/Engine/Plugins/Enterprise/VariantManagerContent/Source/VariantManagerContent/Private/LevelVariantSets.cpp
--- a/Engine/Plugins/Enterprise/VariantManagerContent/Source/VariantManagerContent/Private/LevelVariantSets.cpp
+++ b/Engine/Plugins/Enterprise/VariantManagerContent/Source/VariantManagerContent/Private/LevelVariantSets.cpp
@@ -157,6 +157,11 @@ FString ULevelVariantSets::GetUniqueVariantSetName(const FString& InPrefix)
 		UniqueNames.Add(VariantSet->GetDisplayText().ToString());
 	}
 
+	if (!UniqueNames.Contains(InPrefix))
+	{
+		return InPrefix;
+	}
+
 	FString VarSetName = FString(InPrefix);
 
 	// Remove potentially existing suffix numbers
diff --git a/Engine/Plugins/Enterprise/VariantManagerContent/Source/VariantManagerContent/Private/PropertyValue.cpp b/Engine/Plugins/Enterprise/VariantManagerContent/Source/VariantManagerContent/Private/PropertyValue.cpp
--- a/Engine/Plugins/Enterprise/VariantManagerContent/Source/VariantManagerContent/Private/PropertyValue.cpp
+++ b/Engine/Plugins/Enterprise/VariantManagerContent/Source/VariantManagerContent/Private/PropertyValue.cpp
@@ -112,6 +112,52 @@ void UPropertyValue::Serialize(FArchive& Ar)
 	{
 		Ar << TempObjPtr;
 
+		// Before this version, properties were stored an array of UProperty*. Convert them to
+		// CapturedPropSegment and clear the deprecated arrays
+		if (CustomVersion < FVariantManagerObjectVersion::SerializePropertiesAsNames)
+		{
+			UE_LOG(LogVariantContent, Warning, TEXT("Captured property '%s' was created with an older Unreal Studio version (4.21 or less). A conversion to the new storage format is required and will be attempted. There may be some data loss."), *FullDisplayString);
+
+			int32 NumDeprecatedProps = Properties_DEPRECATED.Num();
+			if (NumDeprecatedProps > 0)
+			{
+				// Back then we didn't store the class directly, and just fetched it from the leaf-most property
+				// Try to do that again as it might help decode ValueBytes if those properties were string types
+				UProperty* LastProp = Properties_DEPRECATED[Properties_DEPRECATED.Num() -1];
+				if (LastProp && LastProp->IsValidLowLevel())
+				{
+					LeafPropertyClass = LastProp->GetClass();
+				}
+
+				CapturedPropSegments.Reserve(NumDeprecatedProps);
+				int32 Index = 0;
+				for (Index = 0; Index < NumDeprecatedProps; Index++)
+				{
+					UProperty* Prop = Properties_DEPRECATED[Index];
+					if (Prop == nullptr || !Prop->IsValidLowLevel() || !PropertyIndices_DEPRECATED.IsValidIndex(Index))
+					{
+						break;
+					}
+
+					FCapturedPropSegment* NewSeg = new(CapturedPropSegments) FCapturedPropSegment;
+					NewSeg->PropertyName = Prop->GetName();
+					NewSeg->PropertyIndex = PropertyIndices_DEPRECATED[Index];
+				}
+
+				// Conversion succeeded
+				if (Index == NumDeprecatedProps)
+				{
+					Properties_DEPRECATED.Reset();
+					PropertyIndices_DEPRECATED.Reset();
+				}
+				else
+				{
+					UE_LOG(LogVariantContent, Warning, TEXT("Failed to convert property '%s'! Captured data will be ignored and property will fail to resolve."), *FullDisplayString);
+					CapturedPropSegments.Reset();
+				}
+			}
+		}
+
 		if (CustomVersion >= FVariantManagerObjectVersion::CorrectSerializationOfFStringBytes)
 		{
 			Ar << TempName;
@@ -976,19 +1022,67 @@ bool UPropertyValue::ResolvePropertiesRecursive(UStruct* ContainerClass, void* C
 					// This lets us search for the component by name instead, ignoring our InnerArrayIndex
 					// This is intuitive because if a component is reordered in the details panel, we kind
 					// of expect our bindings to 'follow'.
-
-					for (int32 ComponentIndex = 0; ComponentIndex < ArrayHelper.Num(); ComponentIndex++)
+					if (!NextSeg.ComponentName.IsEmpty())
 					{
-						void* ObjPtrContainer = ArrayHelper.GetRawPtr(ComponentIndex);
-						UObject* CurrentObject = InnerObjectProperty->GetObjectPropertyValue(ObjPtrContainer);
+						for (int32 ComponentIndex = 0; ComponentIndex < ArrayHelper.Num(); ComponentIndex++)
+						{
+							void* ObjPtrContainer = ArrayHelper.GetRawPtr(ComponentIndex);
+							UObject* CurrentObject = InnerObjectProperty->GetObjectPropertyValue(ObjPtrContainer);
 
-						if (CurrentObject && CurrentObject->IsA(UActorComponent::StaticClass()) && CurrentObject->GetName() == NextSeg.ComponentName)
+							if (CurrentObject && CurrentObject->IsA(UActorComponent::StaticClass()) && CurrentObject->GetName() == NextSeg.ComponentName)
+							{
+								ParentContainerClass = CurrentObject->GetClass();
+								ParentContainerAddress =  CurrentObject;
+
+								// The next link in the chain is just this array's inner. Let's just skip it instead
+								return ResolvePropertiesRecursive(CurrentObject->GetClass(), CurrentObject, SegmentIndex + 2);
+							}
+						}
+					}
+					// If we're a property recovered from 4.21, we won't have a component name, so we'll have to
+					// try finding our target component by index. We will first check InnerArrayIndex, and if that fails, we
+					// will check the other components until we either find something that resolves or we just fall
+					// out of this scope
+					else
+					{
+						// First check our actual inner array index
+						if (ArrayHelper.IsValidIndex(InnerArrayIndex))
 						{
-							ParentContainerClass = CurrentObject->GetClass();
-							ParentContainerAddress =  CurrentObject;
+							void* ObjPtrContainer = ArrayHelper.GetRawPtr(InnerArrayIndex);
+							UObject* CurrentObject = InnerObjectProperty->GetObjectPropertyValue(ObjPtrContainer);
+							if (CurrentObject && CurrentObject->IsA(UActorComponent::StaticClass()))
+							{
+								if (ResolvePropertiesRecursive(CurrentObject->GetClass(), CurrentObject, SegmentIndex + 2))
+								{
+									ParentContainerClass = CurrentObject->GetClass();
+									ParentContainerAddress =  CurrentObject;
+									NextSeg.ComponentName = CurrentObject->GetName();
+									return true;
+								}
+							}
+						}
 
-							// The next link in the chain is just this array's inner. Let's just skip it instead
-							return ResolvePropertiesRecursive(CurrentObject->GetClass(), CurrentObject, SegmentIndex + 2);
+						// Check every component for something that resolves. It's the best we can do
+						for (int32 ComponentIndex = 0; ComponentIndex < ArrayHelper.Num(); ComponentIndex++)
+						{
+							// Already checked that one
+							if (ComponentIndex == InnerArrayIndex)
+							{
+								continue;
+							}
+
+							void* ObjPtrContainer = ArrayHelper.GetRawPtr(ComponentIndex);
+							UObject* CurrentObject = InnerObjectProperty->GetObjectPropertyValue(ObjPtrContainer);
+							if (CurrentObject && CurrentObject->IsA(UActorComponent::StaticClass()))
+							{
+								if (ResolvePropertiesRecursive(CurrentObject->GetClass(), CurrentObject, SegmentIndex + 2))
+								{
+									ParentContainerClass = CurrentObject->GetClass();
+									ParentContainerAddress =  CurrentObject;
+									NextSeg.ComponentName = CurrentObject->GetName();
+									return true;
+								}
+							}
 						}
 					}
 				}
diff --git a/Engine/Plugins/Enterprise/VariantManagerContent/Source/VariantManagerContent/Private/Variant.cpp b/Engine/Plugins/Enterprise/VariantManagerContent/Source/VariantManagerContent/Private/Variant.cpp
--- a/Engine/Plugins/Enterprise/VariantManagerContent/Source/VariantManagerContent/Private/Variant.cpp
+++ b/Engine/Plugins/Enterprise/VariantManagerContent/Source/VariantManagerContent/Private/Variant.cpp
@@ -29,7 +29,19 @@ void UVariant::Serialize(FArchive& Ar)
 	Ar.UsingCustomVersion(FVariantManagerObjectVersion::GUID);
 	int32 CustomVersion = Ar.CustomVer(FVariantManagerObjectVersion::GUID);
 
-	if (CustomVersion >= FVariantManagerObjectVersion::CategoryFlagsAndManualDisplayText)
+	if (CustomVersion < FVariantManagerObjectVersion::CategoryFlagsAndManualDisplayText)
+	{
+		// Recover name from back when it was an UPROPERTY
+		if (Ar.IsLoading())
+		{
+			if (!DisplayText_DEPRECATED.IsEmpty())
+			{
+				DisplayText = DisplayText_DEPRECATED;
+				DisplayText_DEPRECATED = FText();
+			}
+		}
+	}
+	else
 	{
 		Ar << DisplayText;
 	}
diff --git a/Engine/Plugins/Enterprise/VariantManagerContent/Source/VariantManagerContent/Private/VariantSet.cpp b/Engine/Plugins/Enterprise/VariantManagerContent/Source/VariantManagerContent/Private/VariantSet.cpp
--- a/Engine/Plugins/Enterprise/VariantManagerContent/Source/VariantManagerContent/Private/VariantSet.cpp
+++ b/Engine/Plugins/Enterprise/VariantManagerContent/Source/VariantManagerContent/Private/VariantSet.cpp
@@ -29,7 +29,19 @@ void UVariantSet::Serialize(FArchive& Ar)
 	Ar.UsingCustomVersion(FVariantManagerObjectVersion::GUID);
 	int32 CustomVersion = Ar.CustomVer(FVariantManagerObjectVersion::GUID);
 
-	if (CustomVersion >= FVariantManagerObjectVersion::CategoryFlagsAndManualDisplayText)
+	if (CustomVersion < FVariantManagerObjectVersion::CategoryFlagsAndManualDisplayText)
+	{
+		// Recover name from back when it was an UPROPERTY
+		if (Ar.IsLoading())
+		{
+			if (!DisplayText_DEPRECATED.IsEmpty())
+			{
+				DisplayText = DisplayText_DEPRECATED;
+				DisplayText_DEPRECATED = FText();
+			}
+		}
+	}
+	else
 	{
 		Ar << DisplayText;
 	}
@@ -65,6 +77,11 @@ FString UVariantSet::GetUniqueVariantName(const FString& InPrefix)
 		UniqueNames.Add(Variant->GetDisplayText().ToString());
 	}
 
+	if (!UniqueNames.Contains(InPrefix))
+	{
+		return InPrefix;
+	}
+
 	FString VarName = FString(InPrefix);
 
 	// Remove potentially existing suffix numbers
diff --git a/Engine/Plugins/Enterprise/VariantManagerContent/Source/VariantManagerContent/Public/PropertyValue.h b/Engine/Plugins/Enterprise/VariantManagerContent/Source/VariantManagerContent/Public/PropertyValue.h
--- a/Engine/Plugins/Enterprise/VariantManagerContent/Source/VariantManagerContent/Public/PropertyValue.h
+++ b/Engine/Plugins/Enterprise/VariantManagerContent/Source/VariantManagerContent/Public/PropertyValue.h
@@ -42,8 +42,6 @@ struct FCapturedPropSegment
 {
 	GENERATED_USTRUCT_BODY()
 
-	// These are properties so they get duplicated automatically when we duplicate bindings, variants, etc
-
 	UPROPERTY()
 	FString PropertyName;
 
@@ -144,6 +142,14 @@ protected:
 	uint8* PropertyValuePtr;
 	UFunction* PropertySetter;
 
+	// Properties were previously stored like this. Use CapturedPropSegments from now on, which stores
+	// properties by name instead. It is much safer, as we can't guarantee these pointers will be valid
+	// if they point at other packages (will depend on package load order, etc).
+	UPROPERTY()
+	TArray<UProperty*> Properties_DEPRECATED;
+	UPROPERTY()
+	TArray<int32> PropertyIndices_DEPRECATED;
+
 	UPROPERTY()
 	TArray<FCapturedPropSegment> CapturedPropSegments;
 
diff --git a/Engine/Plugins/Enterprise/VariantManagerContent/Source/VariantManagerContent/Public/Variant.h b/Engine/Plugins/Enterprise/VariantManagerContent/Source/VariantManagerContent/Public/Variant.h
--- a/Engine/Plugins/Enterprise/VariantManagerContent/Source/VariantManagerContent/Public/Variant.h
+++ b/Engine/Plugins/Enterprise/VariantManagerContent/Source/VariantManagerContent/Public/Variant.h
@@ -50,6 +50,10 @@ public:
 
 private:
 
+	// The display name used to be a property. Use the non-deprecated, non-property version from now on
+	UPROPERTY()
+	FText DisplayText_DEPRECATED;
+
 	FText DisplayText;
 
 	UPROPERTY()
diff --git a/Engine/Plugins/Enterprise/VariantManagerContent/Source/VariantManagerContent/Public/VariantSet.h b/Engine/Plugins/Enterprise/VariantManagerContent/Source/VariantManagerContent/Public/VariantSet.h
--- a/Engine/Plugins/Enterprise/VariantManagerContent/Source/VariantManagerContent/Public/VariantSet.h
+++ b/Engine/Plugins/Enterprise/VariantManagerContent/Source/VariantManagerContent/Public/VariantSet.h
@@ -54,6 +54,10 @@ public:
 
 private:
 
+	// The display name used to be a property. Use the non-deprecated, non-property version from now on
+	UPROPERTY()
+	FText DisplayText_DEPRECATED;
+
 	FText DisplayText;
 
 	UPROPERTY()
