Merging CL #7708931 from //UE4/Dev-Enterprise to //UE4/Release-4.23

VariantManager: Add nullptr checks on blueprint-exposed functions
#rnx

diff --git a/Engine/Plugins/Enterprise/VariantManagerContent/Source/VariantManagerContent/Private/LevelVariantSets.cpp b/Engine/Plugins/Enterprise/VariantManagerContent/Source/VariantManagerContent/Private/LevelVariantSets.cpp
--- a/Engine/Plugins/Enterprise/VariantManagerContent/Source/VariantManagerContent/Private/LevelVariantSets.cpp
+++ b/Engine/Plugins/Enterprise/VariantManagerContent/Source/VariantManagerContent/Private/LevelVariantSets.cpp
@@ -71,6 +71,11 @@ void ULevelVariantSets::AddVariantSets(const TArray<UVariantSet*>& NewVariantSet
 	TSet<ULevelVariantSets*> ParentsModified;
 	for (UVariantSet* NewVarSet : NewVariantSets)
 	{
+		if (NewVarSet == nullptr)
+		{
+			continue;
+		}
+
 		ULevelVariantSets* OldParent = NewVarSet->GetParent();
 
 		if (OldParent)
@@ -131,6 +136,11 @@ void ULevelVariantSets::AddVariantSets(const TArray<UVariantSet*>& NewVariantSet
 
 int32 ULevelVariantSets::GetVariantSetIndex(UVariantSet* VarSet)
 {
+	if (VarSet == nullptr)
+	{
+		return INDEX_NONE;
+	}
+
 	return VariantSets.Find(VarSet);
 }
 
diff --git a/Engine/Plugins/Enterprise/VariantManagerContent/Source/VariantManagerContent/Private/Variant.cpp b/Engine/Plugins/Enterprise/VariantManagerContent/Source/VariantManagerContent/Private/Variant.cpp
--- a/Engine/Plugins/Enterprise/VariantManagerContent/Source/VariantManagerContent/Private/Variant.cpp
+++ b/Engine/Plugins/Enterprise/VariantManagerContent/Source/VariantManagerContent/Private/Variant.cpp
@@ -75,6 +75,11 @@ void UVariant::AddBindings(const TArray<UVariantObjectBinding*>& NewBindings, in
 	TSet<UVariant*> ParentsModified;
 	for (UVariantObjectBinding* NewBinding : NewBindings)
 	{
+		if (NewBinding == nullptr)
+		{
+			continue;
+		}
+
 		UVariant* OldParent = NewBinding->GetParent();
 
 		if (OldParent)
@@ -139,6 +144,11 @@ void UVariant::AddBindings(const TArray<UVariantObjectBinding*>& NewBindings, in
 
 int32 UVariant::GetBindingIndex(UVariantObjectBinding* Binding)
 {
+	if (Binding == nullptr)
+	{
+		return INDEX_NONE;
+	}
+
 	return ObjectBindings.Find(Binding);
 }
 
diff --git a/Engine/Plugins/Enterprise/VariantManagerContent/Source/VariantManagerContent/Private/VariantObjectBinding.cpp b/Engine/Plugins/Enterprise/VariantManagerContent/Source/VariantManagerContent/Private/VariantObjectBinding.cpp
--- a/Engine/Plugins/Enterprise/VariantManagerContent/Source/VariantManagerContent/Private/VariantObjectBinding.cpp
+++ b/Engine/Plugins/Enterprise/VariantManagerContent/Source/VariantManagerContent/Private/VariantObjectBinding.cpp
@@ -134,6 +134,11 @@ void UVariantObjectBinding::AddCapturedProperties(const TArray<UPropertyValue*>&
 	TSet<UVariantObjectBinding*> ParentsModified;
 	for (UPropertyValue* NewProp : NewProperties)
 	{
+		if (NewProp == nullptr)
+		{
+			continue;
+		}
+
 		if (ExistingProperties.Contains(NewProp->GetFullDisplayString()))
 		{
 			continue;
diff --git a/Engine/Plugins/Enterprise/VariantManagerContent/Source/VariantManagerContent/Private/VariantSet.cpp b/Engine/Plugins/Enterprise/VariantManagerContent/Source/VariantManagerContent/Private/VariantSet.cpp
--- a/Engine/Plugins/Enterprise/VariantManagerContent/Source/VariantManagerContent/Private/VariantSet.cpp
+++ b/Engine/Plugins/Enterprise/VariantManagerContent/Source/VariantManagerContent/Private/VariantSet.cpp
@@ -129,6 +129,11 @@ void UVariantSet::AddVariants(const TArray<UVariant*>& NewVariants, int32 Index)
 	TSet<UVariantSet*> ParentsModified;
 	for (UVariant* NewVariant : NewVariants)
 	{
+		if (NewVariant == nullptr)
+		{
+			continue;
+		}
+
 		UVariantSet* OldParent = NewVariant->GetParent();
 
 		// We can't just RemoveBinding since that might remove the wrong item
@@ -195,6 +200,11 @@ void UVariantSet::AddVariants(const TArray<UVariant*>& NewVariants, int32 Index)
 
 int32 UVariantSet::GetVariantIndex(UVariant* Var)
 {
+	if (Var == nullptr)
+	{
+		return INDEX_NONE;
+	}
+
 	return Variants.Find(Var);
 }
 
