VariantManager: Fix RelativeLocation/Scale not checking for current using tolerance #jira UE-81892 #rb JeanMichel.Dignard

diff --git a/Engine/Plugins/Enterprise/VariantManagerContent/Source/VariantManagerContent/Private/PropertyValue.cpp b/Engine/Plugins/Enterprise/VariantManagerContent/Source/VariantManagerContent/Private/PropertyValue.cpp
--- a/Engine/Plugins/Enterprise/VariantManagerContent/Source/VariantManagerContent/Private/PropertyValue.cpp
+++ b/Engine/Plugins/Enterprise/VariantManagerContent/Source/VariantManagerContent/Private/PropertyValue.cpp
@@ -1229,6 +1229,13 @@ bool UPropertyValue::IsRecordedDataCurrent()
 
 		return RecordedRotator->Equals(*CurrentRotator, SCENECOMPONENT_ROTATOR_TOLERANCE);
 	}
+	else if (PropCategory == EPropertyValueCategory::RelativeLocation ||  PropCategory == EPropertyValueCategory::RelativeScale3D)
+	{
+		const FVector* RecordedVec = (const FVector*)RecordedData.GetData();
+		const FVector* CurrentVec = (const FVector*)CurrentData.GetData();
+
+		return RecordedVec->Equals(*CurrentVec);
+	}
 
 	return RecordedData == CurrentData;
 }
