VariantManager: Fix problems caused by accessing deleted objects #jira UE-82175, UE-82218 #rb JeanMichel.Dignard

diff --git a/Engine/Plugins/Enterprise/VariantManagerContent/Source/VariantManagerContent/Private/PropertyValue.cpp b/Engine/Plugins/Enterprise/VariantManagerContent/Source/VariantManagerContent/Private/PropertyValue.cpp
--- a/Engine/Plugins/Enterprise/VariantManagerContent/Source/VariantManagerContent/Private/PropertyValue.cpp
+++ b/Engine/Plugins/Enterprise/VariantManagerContent/Source/VariantManagerContent/Private/PropertyValue.cpp
@@ -551,12 +551,9 @@ TArray<uint8> UPropertyValue::GetDataFromResolvedObject() const
 
 void UPropertyValue::RecordDataFromResolvedObject()
 {
-	if (!HasValidResolve())
+	if (!Resolve())
 	{
-		if (!Resolve())
-		{
-			return;
-		}
+		return;
 	}
 
 	TArray<uint8> NewData = GetDataFromResolvedObject();
@@ -1178,12 +1175,9 @@ void UPropertyValue::ClearDefaultValue()
 
 bool UPropertyValue::IsRecordedDataCurrent()
 {
-	if (!HasValidResolve())
+	if (!Resolve())
 	{
-		if (!Resolve())
-		{
-			return false;
-		}
+		return false;
 	}
 
 	if (!HasRecordedData())
