Fixing CIS error after Dev-Core copy up

diff --git a/Engine/Plugins/Enterprise/VariantManagerContent/Source/VariantManagerContent/Public/PropertyValue.h b/Engine/Plugins/Enterprise/VariantManagerContent/Source/VariantManagerContent/Public/PropertyValue.h
--- a/Engine/Plugins/Enterprise/VariantManagerContent/Source/VariantManagerContent/Public/PropertyValue.h
+++ b/Engine/Plugins/Enterprise/VariantManagerContent/Source/VariantManagerContent/Public/PropertyValue.h
@@ -5,12 +5,11 @@
 #include "CoreMinimal.h"
 #include "UObject/ObjectMacros.h"
 #include "UObject/Field.h"
+#include "PropertyValue.generated.h"
 
-// WARNING: This should always be the last include in any file that needs it (except .generated.h)
+// WARNING: This should always be the last include in any file that needs it
 #include "UObject/UndefineUPropertyMacros.h"
 
-#include "PropertyValue.generated.h"
-
 #define PATH_DELIMITER TEXT(" / ")
 #define ATTACH_CHILDREN_NAME TEXT("Children")
 
diff --git a/Engine/Source/Programs/UnrealHeaderTool/Private/HeaderParser.cpp b/Engine/Source/Programs/UnrealHeaderTool/Private/HeaderParser.cpp
--- a/Engine/Source/Programs/UnrealHeaderTool/Private/HeaderParser.cpp
+++ b/Engine/Source/Programs/UnrealHeaderTool/Private/HeaderParser.cpp
@@ -9093,7 +9093,7 @@ void FHeaderParser::SimplifiedClassParse(const TCHAR* Filename, const TCHAR* InB
 			// Handle #include directives as if they were 'dependson' keywords.
 			FString DependsOnHeaderName = Str;
 
-			if (DependsOnHeaderName != TEXT("\"UObject/DefineUPropertyMacros.h\""))
+			if (DependsOnHeaderName != TEXT("\"UObject/DefineUPropertyMacros.h\"") && DependsOnHeaderName != TEXT("\"UObject/UndefineUPropertyMacros.h\""))
 			{
 				if (bFoundGeneratedInclude)
 				{
