Reverting bad merge

diff --git a/Engine/Plugins/Enterprise/VariantManagerContent/Source/VariantManagerContent/Public/PropertyValue.h b/Engine/Plugins/Enterprise/VariantManagerContent/Source/VariantManagerContent/Public/PropertyValue.h
--- a/Engine/Plugins/Enterprise/VariantManagerContent/Source/VariantManagerContent/Public/PropertyValue.h
+++ b/Engine/Plugins/Enterprise/VariantManagerContent/Source/VariantManagerContent/Public/PropertyValue.h
@@ -6,8 +6,7 @@
 #include "UObject/ObjectMacros.h"
 #include "UObject/Field.h"
 
-// WARNING: This should always be the last include in any file that needs it
-#include "UObject/UndefineUPropertyMacros.h"
+#include "PropertyValue.generated.h"
 
 #define PATH_DELIMITER TEXT(" / ")
 #define ATTACH_CHILDREN_NAME TEXT("Children")
@@ -17,7 +16,7 @@ DECLARE_MULTICAST_DELEGATE(FOnPropertyApplied);
 
 class UVariantObjectBinding;
 class USCS_Node;
-class UProperty;
+class FProperty;
 
 UENUM()
 enum class EPropertyValueCategory : uint8
@@ -248,5 +247,4 @@ class VARIANTMANAGERCONTENT_API UPropertyValueVisibility : public UPropertyValue
 	// UObject Interface
 	virtual void Serialize(FArchive& Ar) override;
 	//~ End UObject Interface
-};
-#include "UObject/DefineUPropertyMacros.h"
\ No newline at end of file
+};
\ No newline at end of file
