VariantManager: Fix back comp issue for function callers made in 4.24

Integrate from Dev-Enterprise cl 11292251

diff --git a/Engine/Plugins/Enterprise/VariantManager/Source/VariantManager/Private/DisplayNodes/VariantManagerFunctionPropertyNode.cpp b/Engine/Plugins/Enterprise/VariantManager/Source/VariantManager/Private/DisplayNodes/VariantManagerFunctionPropertyNode.cpp
--- a/Engine/Plugins/Enterprise/VariantManager/Source/VariantManager/Private/DisplayNodes/VariantManagerFunctionPropertyNode.cpp
+++ b/Engine/Plugins/Enterprise/VariantManager/Source/VariantManager/Private/DisplayNodes/VariantManagerFunctionPropertyNode.cpp
@@ -252,7 +252,7 @@ TSharedRef<SWidget> FVariantManagerFunctionPropertyNode::GetMenuContent()
 				}
 				else
 				{
-					TooltipText = LOCTEXT("InvalidTooltip", "Function must be valid and have either no input pins or 4 input pins with the correct types:\nThe bound actor's class, Level Variant Sets, Variant Set and Variant");
+					TooltipText = LOCTEXT("InvalidTooltip", "Function must be valid and have either: \n - Zero input pins;\n - One input pin (with type matching the bound actor's class);\n - 4 input pins (with types: the bound actor's class, Level Variant Sets, Variant Set and Variant)");
 				}
 			}
 			else
diff --git a/Engine/Plugins/Enterprise/VariantManagerContent/Source/VariantManagerContent/Private/FunctionCaller.cpp b/Engine/Plugins/Enterprise/VariantManagerContent/Source/VariantManagerContent/Private/FunctionCaller.cpp
--- a/Engine/Plugins/Enterprise/VariantManagerContent/Source/VariantManagerContent/Private/FunctionCaller.cpp
+++ b/Engine/Plugins/Enterprise/VariantManagerContent/Source/VariantManagerContent/Private/FunctionCaller.cpp
@@ -33,7 +33,6 @@ bool FFunctionCaller::IsBoundToBlueprint() const
 
 bool FFunctionCaller::IsValidFunction(UK2Node_FunctionEntry* Function)
 {
-	// Must have a single non-reference parameter on the function
 	if (!Function || Function->IsPendingKill() || Function->GetGraph()->IsPendingKill())
 	{
 		return false;
@@ -42,41 +41,52 @@ bool FFunctionCaller::IsValidFunction(UK2Node_FunctionEntry* Function)
 	{
 		return true;
 	}
-	else if (Function->UserDefinedPins.Num() != 4)
+	else if (Function->UserDefinedPins.Num() == 1)
 	{
-		return false;
-	}
-
-	TSharedPtr<FUserPinInfo> TargetPin			 = Function->UserDefinedPins[0];
-	TSharedPtr<FUserPinInfo> LevelVariantSetsPin = Function->UserDefinedPins[1];
-	TSharedPtr<FUserPinInfo> VariantSetPin		 = Function->UserDefinedPins[2];
-	TSharedPtr<FUserPinInfo> VariantPin			 = Function->UserDefinedPins[3];
+		TSharedPtr<FUserPinInfo> TargetPin = Function->UserDefinedPins[0];
+		if (TargetPin->PinType.bIsReference ||
+			TargetPin->PinType.PinCategory != UEdGraphSchema_K2::PC_Object)
+		{
+			return false;
+		}
 
-	if (TargetPin->PinType.bIsReference ||
-		TargetPin->PinType.PinCategory != UEdGraphSchema_K2::PC_Object)
-	{
-		return false;
-	}
-	if (LevelVariantSetsPin->PinType.bIsReference ||
-		LevelVariantSetsPin->PinType.PinCategory != UEdGraphSchema_K2::PC_Object ||
-		LevelVariantSetsPin->PinType.PinSubCategoryObject != ULevelVariantSets::StaticClass())
-	{
-		return false;
-	}
-	if (VariantSetPin->PinType.bIsReference ||
-		VariantSetPin->PinType.PinCategory != UEdGraphSchema_K2::PC_Object ||
-		VariantSetPin->PinType.PinSubCategoryObject != UVariantSet::StaticClass())
-	{
-		return false;
+		return true;
 	}
-	if (VariantPin->PinType.bIsReference ||
-		VariantPin->PinType.PinCategory != UEdGraphSchema_K2::PC_Object ||
-		VariantPin->PinType.PinSubCategoryObject != UVariant::StaticClass())
+	else if (Function->UserDefinedPins.Num() == 4)
 	{
-		return false;
+		TSharedPtr<FUserPinInfo> TargetPin			 = Function->UserDefinedPins[0];
+		TSharedPtr<FUserPinInfo> LevelVariantSetsPin = Function->UserDefinedPins[1];
+		TSharedPtr<FUserPinInfo> VariantSetPin		 = Function->UserDefinedPins[2];
+		TSharedPtr<FUserPinInfo> VariantPin			 = Function->UserDefinedPins[3];
+
+		if (TargetPin->PinType.bIsReference ||
+			TargetPin->PinType.PinCategory != UEdGraphSchema_K2::PC_Object)
+		{
+			return false;
+		}
+		if (LevelVariantSetsPin->PinType.bIsReference ||
+			LevelVariantSetsPin->PinType.PinCategory != UEdGraphSchema_K2::PC_Object ||
+			LevelVariantSetsPin->PinType.PinSubCategoryObject != ULevelVariantSets::StaticClass())
+		{
+			return false;
+		}
+		if (VariantSetPin->PinType.bIsReference ||
+			VariantSetPin->PinType.PinCategory != UEdGraphSchema_K2::PC_Object ||
+			VariantSetPin->PinType.PinSubCategoryObject != UVariantSet::StaticClass())
+		{
+			return false;
+		}
+		if (VariantPin->PinType.bIsReference ||
+			VariantPin->PinType.PinCategory != UEdGraphSchema_K2::PC_Object ||
+			VariantPin->PinType.PinSubCategoryObject != UVariant::StaticClass())
+		{
+			return false;
+		}
+
+		return true;
 	}
 
-	return true;
+	return false;
 }
 
 uint32 FFunctionCaller::GetDisplayOrder() const
@@ -118,22 +128,4 @@ void FFunctionCaller::PostSerialize(const FArchive& Ar)
 #endif
 }
 
-bool FFunctionCaller::IsValidFunction(UFunction* Function)
-{
-	// Must have a single non-reference parameter on the function
-	if (!Function)
-	{
-		return false;
-	}
-	else if (Function->NumParms == 0)
-	{
-		return true;
-	}
-	else if (Function->NumParms != 1 || !Function->PropertyLink || (Function->PropertyLink->GetPropertyFlags() & CPF_ReferenceParm))
-	{
-		return false;
-	}
 
-	// Parameter must be an object or interface property
-	return Function->PropertyLink->IsA<FObjectProperty>() || Function->PropertyLink->IsA<FInterfaceProperty>();
-}
diff --git a/Engine/Plugins/Enterprise/VariantManagerContent/Source/VariantManagerContent/Private/VariantObjectBinding.cpp b/Engine/Plugins/Enterprise/VariantManagerContent/Source/VariantManagerContent/Private/VariantObjectBinding.cpp
--- a/Engine/Plugins/Enterprise/VariantManagerContent/Source/VariantManagerContent/Private/VariantObjectBinding.cpp
+++ b/Engine/Plugins/Enterprise/VariantManagerContent/Source/VariantManagerContent/Private/VariantObjectBinding.cpp
@@ -299,6 +299,26 @@ void UVariantObjectBinding::ExecuteTargetFunction(FName FunctionName)
 	{
 		DirectorInstance->ProcessEvent(Func, nullptr);
 	}
+	// Mostly for backwards compatibility
+	else if (Func->NumParms == 1)
+	{
+		if (FObjectProperty* ObjectParameter = CastField<FObjectProperty>(Func->PropertyLink))
+		{
+			if (!ObjectParameter->PropertyClass || BoundObject->IsA(ObjectParameter->PropertyClass))
+			{
+				DirectorInstance->ProcessEvent(Func, &BoundObject);
+			}
+			else
+			{
+				UE_LOG(LogVariantContent, Error, TEXT("Failed to call function '%s' with object '%s' because it is not the correct type. Function expects a '%s' but target object is a '%s'."),
+					*Func->GetName(),
+					*BoundObject->GetName(),
+					*ObjectParameter->PropertyClass->GetName(),
+					*BoundObject->GetClass()->GetName()
+				);
+			}
+		}
+	}
 	else if (Func->NumParms == 4)
 	{
 		// Setup parameters to the Func as a struct
diff --git a/Engine/Plugins/Enterprise/VariantManagerContent/Source/VariantManagerContent/Public/FunctionCaller.h b/Engine/Plugins/Enterprise/VariantManagerContent/Source/VariantManagerContent/Public/FunctionCaller.h
--- a/Engine/Plugins/Enterprise/VariantManagerContent/Source/VariantManagerContent/Public/FunctionCaller.h
+++ b/Engine/Plugins/Enterprise/VariantManagerContent/Source/VariantManagerContent/Public/FunctionCaller.h
@@ -26,12 +26,6 @@ struct VARIANTMANAGERCONTENT_API FFunctionCaller
 	 */
 	void PostSerialize(const FArchive& Ar);
 
-	/**
-	 * Check whether the specified function is valid for a movie scene event
-	 * Functions must have either no parameters, or a single, pass-by-value object/interface parameter, with no return parameter.
-	 */
-	static bool IsValidFunction(UFunction* Function);
-
 	/**
 	 * The function that should be called to invoke this event.
 	 * Functions must have either no parameters, or a single, pass-by-value object/interface parameter, with no return parameter.
