VariantManager: Enable blueprint runtime access to Variant::IsActive

Integrated from Dev-Enterprise cl 11285130

diff --git a/Engine/Plugins/Enterprise/VariantManagerContent/Source/VariantManagerContent/Public/Variant.h b/Engine/Plugins/Enterprise/VariantManagerContent/Source/VariantManagerContent/Public/Variant.h
--- a/Engine/Plugins/Enterprise/VariantManagerContent/Source/VariantManagerContent/Public/Variant.h
+++ b/Engine/Plugins/Enterprise/VariantManagerContent/Source/VariantManagerContent/Public/Variant.h
@@ -49,6 +49,7 @@ public:
 	void SwitchOn();
 
 	// Returns true if none of our properties are dirty
+	UFUNCTION(BlueprintCallable, Category="Variant")
 	bool IsActive();
 
 	// Sets the thumbnail to use for this variant. Can receive nullptr to clear it
