VariantManager: Fixes related to nested SwitchActors

Fix parent SwitchActor clearing choices of child SwitchActors;
Fix child SwitchActors being able to set children as visible when they themselves are unselected;

diff --git a/Engine/Plugins/Enterprise/VariantManagerContent/Source/VariantManagerContent/Private/SwitchActor.cpp b/Engine/Plugins/Enterprise/VariantManagerContent/Source/VariantManagerContent/Private/SwitchActor.cpp
--- a/Engine/Plugins/Enterprise/VariantManagerContent/Source/VariantManagerContent/Private/SwitchActor.cpp
+++ b/Engine/Plugins/Enterprise/VariantManagerContent/Source/VariantManagerContent/Private/SwitchActor.cpp
@@ -5,6 +5,35 @@
 #include "Algo/Sort.h"
 #include "CoreMinimal.h"
 
+namespace SwitchActorImpl
+{
+	int32 DiscoverSelectedOption(const TArray<AActor*>& Actors)
+	{
+		int32 SingleVisibleChildIndex = INDEX_NONE;
+		for (int32 Index = 0; Index < Actors.Num(); ++Index)
+		{
+			const AActor* Actor = Actors[Index];
+			if (USceneComponent* ActorRoot = Actor->GetRootComponent())
+			{
+				if (ActorRoot->IsVisible())
+				{
+					if (SingleVisibleChildIndex == INDEX_NONE)
+					{
+						SingleVisibleChildIndex = Index;
+					}
+					else
+					{
+						SingleVisibleChildIndex = INDEX_NONE;
+						break;
+					}
+				}
+			}
+		}
+
+		return SingleVisibleChildIndex;
+	}
+}
+
 ASwitchActor::ASwitchActor(const FObjectInitializer& Init)
 	: Super(Init)
 {
@@ -16,8 +45,9 @@ ASwitchActor::ASwitchActor(const FObjectInitializer& Init)
 
 TArray<AActor*> ASwitchActor::GetOptions() const
 {
+	const bool bResetArray = false;
 	TArray<AActor*> Result;
-	GetAttachedActors(Result, false);
+	GetAttachedActors(Result, bResetArray);
 
 	// We have to sort these by FName because the attach order is not guaranteed
 	// It seems to invert when going into PIE, for example
@@ -31,30 +61,7 @@ TArray<AActor*> ASwitchActor::GetOptions() const
 
 int32 ASwitchActor::GetSelectedOption() const
 {
-	TArray<AActor*> Actors = GetOptions();
-
-	int32 SingleVisibleChildIndex = INDEX_NONE;
-	for (int32 Index = 0; Index < Actors.Num(); ++Index)
-	{
-		const AActor* Actor = Actors[Index];
-		if (USceneComponent* ActorRoot = Actor->GetRootComponent())
-		{
-			if (ActorRoot->IsVisible())
-			{
-				if (SingleVisibleChildIndex == INDEX_NONE)
-				{
-					SingleVisibleChildIndex = Index;
-				}
-				else
-				{
-					SingleVisibleChildIndex = INDEX_NONE;
-					break;
-				}
-			}
-		}
-	}
-
-	return SingleVisibleChildIndex;
+	return LastSelectedOption;
 }
 
 void ASwitchActor::SelectOption(int32 OptionIndex)
@@ -66,18 +73,32 @@ void ASwitchActor::SelectOption(int32 OptionIndex)
 		return;
 	}
 
+	Modify();
+
+	// So that we set our children to visible only if we ourselves are visible,
+	// which helps this actor behave nicely when child of another switchactor that hasn't selected us
+	bool bOurVisibility = this->GetRootComponent()->GetVisibleFlag();
+
 	for (int32 Index = 0; Index < Actors.Num(); ++Index)
 	{
-		if (AActor* Actor = Actors[Index])
+		AActor* Actor = Actors[Index];
+		if (ASwitchActor* ChildSwitchActor = Cast<ASwitchActor>(Actor))
+		{
+			ChildSwitchActor->SetVisibility(bOurVisibility && Index == OptionIndex);
+		}
+		else if(Actor)
 		{
 			if (USceneComponent* ActorRootComponent = Actor->GetRootComponent())
 			{
 				ActorRootComponent->Modify();
-				ActorRootComponent->SetVisibility(Index == OptionIndex, true);
+
+				const bool bPropagateToChildren = true;
+				ActorRootComponent->SetVisibility(bOurVisibility && Index == OptionIndex, bPropagateToChildren);
 			}
 		}
 	}
 
+	LastSelectedOption = OptionIndex;
 	OnSwitchActorSwitch.Broadcast(OptionIndex);
 }
 
@@ -86,4 +107,22 @@ FOnSwitchActorSwitch& ASwitchActor::GetOnSwitchDelegate()
 	return OnSwitchActorSwitch;
 }
 
+void ASwitchActor::SetVisibility(bool bVisible)
+{
+	Modify();
+
+	const bool bPropagateToChildren = true;
+	this->GetRootComponent()->SetVisibility(bVisible, bPropagateToChildren);
+	SelectOption(LastSelectedOption);
+}
+
+void ASwitchActor::PostLoad()
+{
+	Super::PostLoad();
+
+	if (LastSelectedOption == INDEX_NONE)
+	{
+		LastSelectedOption = SwitchActorImpl::DiscoverSelectedOption(GetOptions());
+	}
+}
 
diff --git a/Engine/Plugins/Enterprise/VariantManagerContent/Source/VariantManagerContent/Public/SwitchActor.h b/Engine/Plugins/Enterprise/VariantManagerContent/Source/VariantManagerContent/Public/SwitchActor.h
--- a/Engine/Plugins/Enterprise/VariantManagerContent/Source/VariantManagerContent/Public/SwitchActor.h
+++ b/Engine/Plugins/Enterprise/VariantManagerContent/Source/VariantManagerContent/Public/SwitchActor.h
@@ -36,9 +36,18 @@ public:
 	FOnSwitchActorSwitch& GetOnSwitchDelegate();
 
 private:
+	// Dedicated function to set our visibility so that we can restore our component
+	// hierarchy to the last-set option, in case e.g. we're overriden by a parent ASwitchActor
+	void SetVisibility(bool bVisible);
+
+	void PostLoad() override;
+
 	FOnSwitchActorSwitch OnSwitchActorSwitch;
 
 	// Exposing our root component like this allows manual Mobility control on the details panel
 	UPROPERTY(Category = SwitchActor, VisibleAnywhere)
 	class USceneComponent* SceneComponent;
+
+	UPROPERTY()
+	int32 LastSelectedOption;
 };
