VariantManager: Fix crash when loading variants in standalone mode

diff --git a/Engine/Plugins/Enterprise/VariantManagerContent/Source/VariantManagerContent/Private/LevelVariantSets.cpp b/Engine/Plugins/Enterprise/VariantManagerContent/Source/VariantManagerContent/Private/LevelVariantSets.cpp
--- a/Engine/Plugins/Enterprise/VariantManagerContent/Source/VariantManagerContent/Private/LevelVariantSets.cpp
+++ b/Engine/Plugins/Enterprise/VariantManagerContent/Source/VariantManagerContent/Private/LevelVariantSets.cpp
@@ -363,6 +363,7 @@ UWorld* ULevelVariantSets::ComputeCurrentWorld(int32& OutPIEInstanceID)
 	UWorld* EditorWorld = nullptr;
 	for (const FWorldContext& Context : GEngine->GetWorldContexts())
 	{
+		// Return the first PIE world that we can find
 		if (Context.WorldType == EWorldType::PIE)
 		{
 			UWorld* ThisWorld = Context.World();
@@ -372,7 +373,9 @@ UWorld* ULevelVariantSets::ComputeCurrentWorld(int32& OutPIEInstanceID)
 				return ThisWorld;
 			}
 		}
-		else if (Context.WorldType == EWorldType::Editor)
+		// Or else return a valid Editor world. For "Standalone mode" the world type is Game.
+		// Note that this code won't run in an actual packaged build though, as its inside an #if WITH_EDITOR block
+		else if (Context.WorldType == EWorldType::Editor | Context.WorldType == EWorldType::Game)
 		{
 			EditorWorld = Context.World();
 		}
