VariantManager: Mark classes as BlueprintTypes so that they be constructed from Editor Utility Blueprints

diff --git a/Engine/Plugins/Enterprise/VariantManagerContent/Source/VariantManagerContent/Public/LevelVariantSets.h b/Engine/Plugins/Enterprise/VariantManagerContent/Source/VariantManagerContent/Public/LevelVariantSets.h
--- a/Engine/Plugins/Enterprise/VariantManagerContent/Source/VariantManagerContent/Public/LevelVariantSets.h
+++ b/Engine/Plugins/Enterprise/VariantManagerContent/Source/VariantManagerContent/Public/LevelVariantSets.h
@@ -7,13 +7,13 @@
 
 #include "LevelVariantSets.generated.h"
 
-class UVariantSet;
-class ULevelVariantSetsFunctionDirector;
 class ALevelVariantSetsActor;
-class UBlueprintGeneratedClass;
 class UBlueprint;
+class UBlueprintGeneratedClass;
+class ULevelVariantSetsFunctionDirector;
+class UVariantSet;
 
-UCLASS(DefaultToInstanced)
+UCLASS(BlueprintType)
 class VARIANTMANAGERCONTENT_API ULevelVariantSets : public UObject
 {
 	GENERATED_UCLASS_BODY()
diff --git a/Engine/Plugins/Enterprise/VariantManagerContent/Source/VariantManagerContent/Public/Variant.h b/Engine/Plugins/Enterprise/VariantManagerContent/Source/VariantManagerContent/Public/Variant.h
--- a/Engine/Plugins/Enterprise/VariantManagerContent/Source/VariantManagerContent/Public/Variant.h
+++ b/Engine/Plugins/Enterprise/VariantManagerContent/Source/VariantManagerContent/Public/Variant.h
@@ -12,7 +12,7 @@
 
 class UVariantObjectBinding;
 
-UCLASS(DefaultToInstanced)
+UCLASS(BlueprintType)
 class VARIANTMANAGERCONTENT_API UVariant : public UObject
 {
 	GENERATED_UCLASS_BODY()
