VariantManager: Upgrade other struct types from float to double if needed.

diff --git a/Engine/Plugins/Enterprise/VariantManagerContent/Source/VariantManagerContent/Private/PropertyValue.cpp b/Engine/Plugins/Enterprise/VariantManagerContent/Source/VariantManagerContent/Private/PropertyValue.cpp
--- a/Engine/Plugins/Enterprise/VariantManagerContent/Source/VariantManagerContent/Private/PropertyValue.cpp
+++ b/Engine/Plugins/Enterprise/VariantManagerContent/Source/VariantManagerContent/Private/PropertyValue.cpp
@@ -36,7 +36,7 @@ namespace UE
 		namespace Private
 		{
 			template<typename OldType, typename NewType>
-			void ConvertRecordedValueSizes( size_t TargetStructSize, bool& bHasRecordedData, TArray<uint8>& ValueBytes )
+			void ConvertRecordedValueSizes( size_t TargetStructSize, TArray<uint8>& ValueBytes )
 			{
 				const int32 NumElements = ValueBytes.Num() / sizeof( OldType );
 				ensure( ValueBytes.Num() % sizeof( OldType ) == 0 );
@@ -54,34 +54,59 @@ namespace UE
 				}
 
 				ValueBytes = MoveTemp( ConvertedRecordedData );
-				bHasRecordedData = true;
+			}
+
+			size_t GetTargetStructElementSize( UScriptStruct* StructClass )
+			{
+				if ( StructClass )
+				{
+					FName StructName = StructClass->GetFName();
+					if ( StructName == NAME_Vector )
+					{
+						return sizeof( FVector::X );
+					}
+					else if ( StructName == NAME_Rotator )
+					{
+						return sizeof( FRotator::Pitch );
+					}
+					else if ( StructName == NAME_Quat )
+					{
+						return sizeof( FQuat::X );
+					}
+					else if ( StructName == NAME_Vector4 )
+					{
+						return sizeof( FVector4::X );
+					}
+					else if ( StructName == NAME_Vector2D )
+					{
+						return sizeof( FVector2D::X );
+					}
+				}
+
+				return 0;
 			}
 
 			// The purpose of this function is to upgrade ValueBytes to the correct size/values if we last saved our e.g. FVector
 			// recorded data when FVector contained floats, but now it should hold doubles (i.e. ValueBytes holds 12 bytes, but
 			// it should really hold 24 now). This change can happen both ways because UE_LARGE_WORLD_COORDINATES_DISABLED can be defined or not.
-			void UpdateRecordedDataSizesIfNeeded( UScriptStruct* StructClass, bool& bHasRecordedData, TArray<uint8>& ValueBytes )
+			void UpdateRecordedDataSizesIfNeeded( UScriptStruct* StructClass, TArray<uint8>& ValueBytes )
 			{
-				if ( StructClass && bHasRecordedData )
+				if ( !StructClass )
 				{
-					FName StructName = StructClass->GetFName();
-					if ( StructName == NAME_Vector || StructName == NAME_Quat )
-					{
-						const size_t TargetStructSize = StructClass->GetCppStructOps()->GetSize();
+					return;
+				}
 
-						// If these sizes are different we need to update our RecordedData
-						if ( TargetStructSize != ValueBytes.Num() )
-						{
-							const size_t NewElementSize = StructName == NAME_Vector ? sizeof( FVector().X ) : sizeof( FQuat().X );
-							if ( NewElementSize == sizeof( double ) )
-							{
-								ConvertRecordedValueSizes<float, double>( TargetStructSize, bHasRecordedData, ValueBytes );
-							}
-							else if ( NewElementSize == sizeof( float ) )
-							{
-								ConvertRecordedValueSizes<double, float>( TargetStructSize, bHasRecordedData, ValueBytes );
-							}
-						}
+				const size_t TargetStructSize = StructClass->GetCppStructOps()->GetSize();
+				const size_t TargetElementSize = GetTargetStructElementSize( StructClass );
+				if ( TargetElementSize && TargetStructSize && TargetStructSize != ValueBytes.Num() )
+				{
+					if ( TargetElementSize == sizeof( double ) )
+					{
+						ConvertRecordedValueSizes<float, double>( TargetStructSize, ValueBytes );
+					}
+					else if ( TargetElementSize == sizeof( float ) )
+					{
+						ConvertRecordedValueSizes<double, float>( TargetStructSize, ValueBytes );
 					}
 				}
 			}
@@ -531,11 +556,11 @@ bool UPropertyValue::Resolve(UObject* Object)
 		}
 	}
 
-	if ( bStartedUnresolved && HasValidResolve() )
+	if ( bStartedUnresolved && HasValidResolve() && bHasRecordedData )
 	{
 		// We can only do this after we resolve because we need to know the struct property's struct class, which also
 		// means we can't do it on Serialize()
-		UE::PropertyValue::Private::UpdateRecordedDataSizesIfNeeded( GetStructPropertyStruct(), bHasRecordedData, ValueBytes );
+		UE::PropertyValue::Private::UpdateRecordedDataSizesIfNeeded( GetStructPropertyStruct(), ValueBytes );
 	}
 
 	return true;
@@ -561,7 +586,7 @@ bool UPropertyValue::HasValidResolve() const
 			if (this == nullptr ||
 				!GUObjectArray.IsValid(this) ||
 				!Container->GetClass() ||
-				!IsValidChecked(Container) || 
+				!IsValidChecked(Container) ||
 				Container->IsUnreachable() ||
 				Container->HasAnyFlags(RF_BeginDestroyed | RF_FinishDestroyed))
 			{
