USD: Fully wait for texture compilation before trying to use it for a thumbnail or else we end up with the default texture

diff --git a/Engine/Plugins/Enterprise/VariantManagerContent/Source/VariantManagerContent/Private/ThumbnailGenerator.cpp b/Engine/Plugins/Enterprise/VariantManagerContent/Source/VariantManagerContent/Private/ThumbnailGenerator.cpp
--- a/Engine/Plugins/Enterprise/VariantManagerContent/Source/VariantManagerContent/Private/ThumbnailGenerator.cpp
+++ b/Engine/Plugins/Enterprise/VariantManagerContent/Source/VariantManagerContent/Private/ThumbnailGenerator.cpp
@@ -13,6 +13,7 @@
 #include "LegacyScreenPercentageDriver.h"
 #include "Rendering/Texture2DResource.h"
 #include "RenderUtils.h"
+#include "TextureCompiler.h"
 
 #if WITH_EDITOR
 #include "ObjectTools.h"
@@ -135,6 +136,7 @@ UTexture2D* ThumbnailGenerator::GenerateThumbnailFromTexture(UTexture2D* Texture
 	}
 
     // Force all mips to stream in, as we may need to use mip 0 for the thumbnail
+	FTextureCompilingManager::Get().FinishCompilation( { Texture } );
 	Texture->SetForceMipLevelsToBeResident(30.0f);
 	Texture->WaitForStreaming();
 
@@ -142,24 +144,24 @@ UTexture2D* ThumbnailGenerator::GenerateThumbnailFromTexture(UTexture2D* Texture
 	int32 TargetHeight = Texture->GetSizeY();
 
 	if (TargetWidth == 0 || TargetHeight == 0 || !Texture->GetResource())
-		{
+	{
 		UE_LOG(LogVariantContent, Error, TEXT("Failed create a thumbnail from texture '%s'"), *Texture->GetName());
 		return nullptr;
-			}
+	}
 
 	if (TargetWidth > VARIANT_MANAGER_THUMBNAIL_SIZE || TargetHeight > VARIANT_MANAGER_THUMBNAIL_SIZE)
-				{
+	{
 		if (TargetWidth >= TargetHeight)
-			{
+		{
 			TargetHeight = (int)(VARIANT_MANAGER_THUMBNAIL_SIZE * TargetHeight / (float)TargetWidth);
 			TargetWidth = VARIANT_MANAGER_THUMBNAIL_SIZE;
-			}
-			else
-			{
+		}
+		else
+		{
 			TargetWidth = (int)(VARIANT_MANAGER_THUMBNAIL_SIZE * TargetWidth / (float)TargetHeight);
 			TargetHeight = VARIANT_MANAGER_THUMBNAIL_SIZE;
-				}
-			}
+		}
+	}
 
 	UTextureRenderTarget2D* RenderTargetTexture = NewObject<UTextureRenderTarget2D>();
 	RenderTargetTexture->AddToRoot();
@@ -197,15 +199,15 @@ UTexture2D* ThumbnailGenerator::GenerateThumbnailFromTexture(UTexture2D* Texture
 	RenderTargetTexture->RemoveFromRoot();
 	RenderTargetTexture = nullptr;
 
-		const bool bSetSourceData = true;
+	const bool bSetSourceData = true;
 	UTexture2D* Thumbnail = ThumbnailGeneratorImpl::CreateTextureFromBulkData(
-			TargetWidth,
-			TargetHeight,
+		TargetWidth,
+		TargetHeight,
 		(void*)CapturedImage.GetData(),
 		CapturedImage.Num() * sizeof(FColor),
 		PF_B8G8R8A8,
-			bSetSourceData
-		);
+		bSetSourceData
+	);
 
     if (Thumbnail == nullptr)
 	{
