VariantManager: Build fix for non-editor targets.

diff --git a/Engine/Plugins/Enterprise/VariantManagerContent/Source/VariantManagerContent/Private/ThumbnailGenerator.cpp b/Engine/Plugins/Enterprise/VariantManagerContent/Source/VariantManagerContent/Private/ThumbnailGenerator.cpp
--- a/Engine/Plugins/Enterprise/VariantManagerContent/Source/VariantManagerContent/Private/ThumbnailGenerator.cpp
+++ b/Engine/Plugins/Enterprise/VariantManagerContent/Source/VariantManagerContent/Private/ThumbnailGenerator.cpp
@@ -13,9 +13,9 @@
 #include "LegacyScreenPercentageDriver.h"
 #include "Rendering/Texture2DResource.h"
 #include "RenderUtils.h"
-#include "TextureCompiler.h"
 
 #if WITH_EDITOR
+#include "TextureCompiler.h"
 #include "ObjectTools.h"
 #include "LevelEditor.h"
 #endif
@@ -136,7 +136,9 @@ UTexture2D* ThumbnailGenerator::GenerateThumbnailFromTexture(UTexture2D* Texture
 	}
 
     // Force all mips to stream in, as we may need to use mip 0 for the thumbnail
+#if WITH_EDITOR
 	FTextureCompilingManager::Get().FinishCompilation( { Texture } );
+#endif // WITH_EDITOR
 	Texture->SetForceMipLevelsToBeResident(30.0f);
 	Texture->WaitForStreaming();
 
