Fix use of deprecated FCanvas constructor.

diff --git a/Engine/Plugins/Enterprise/VariantManagerContent/Source/VariantManagerContent/Private/ThumbnailGenerator.cpp b/Engine/Plugins/Enterprise/VariantManagerContent/Source/VariantManagerContent/Private/ThumbnailGenerator.cpp
--- a/Engine/Plugins/Enterprise/VariantManagerContent/Source/VariantManagerContent/Private/ThumbnailGenerator.cpp
+++ b/Engine/Plugins/Enterprise/VariantManagerContent/Source/VariantManagerContent/Private/ThumbnailGenerator.cpp
@@ -174,7 +174,7 @@ UTexture2D* ThumbnailGenerator::GenerateThumbnailFromTexture(UTexture2D* Texture
 	FTextureRenderTargetResource* RenderTargetResource = RenderTargetTexture->GameThread_GetRenderTargetResource();
 
 	const double Time = FApp::GetCurrentTime() - GStartTime;
-	FCanvas Canvas(RenderTargetResource, NULL, Time, FApp::GetDeltaTime(), Time, GWorld->Scene->GetFeatureLevel());
+	FCanvas Canvas(RenderTargetResource, NULL, FGameTime::GetTimeSinceAppStart(), GWorld->Scene->GetFeatureLevel());
 	Canvas.Clear(FLinearColor::Black);
 
 	const bool bAlphaBlend = false;
