Crash playing in editor the CarConfigurator.uproject project

diff --git a/Engine/Plugins/Enterprise/VariantManagerContent/Source/VariantManagerContent/Private/PropertyValue.cpp b/Engine/Plugins/Enterprise/VariantManagerContent/Source/VariantManagerContent/Private/PropertyValue.cpp
--- a/Engine/Plugins/Enterprise/VariantManagerContent/Source/VariantManagerContent/Private/PropertyValue.cpp
+++ b/Engine/Plugins/Enterprise/VariantManagerContent/Source/VariantManagerContent/Private/PropertyValue.cpp
@@ -20,6 +20,7 @@
 
 #if WITH_EDITOR
 #include "EdGraphSchema_K2.h"
+#include "Editor.h"
 #endif
 
 // WARNING: This should always be the last include in any file that needs it (except .generated.h)
@@ -303,6 +304,12 @@ UPropertyValue::UPropertyValue(const FObjectInitializer& ObjectInitializer)
 	: Super(ObjectInitializer)
 	, LeafPropertyClass(nullptr)
 {
+#if WITH_EDITOR
+	if (GIsEditor)
+	{
+		FEditorDelegates::EndPIE.AddUObject(this, &UPropertyValue::OnPIEEnded);
+	}
+#endif
 }
 
 void UPropertyValue::Init(const TArray<FCapturedPropSegment>& InCapturedPropSegments, FFieldClass* InLeafPropertyClass, const FString& InFullDisplayString, const FName& InPropertySetterName, EPropertyValueCategory InCategory)
@@ -324,6 +331,24 @@ void UPropertyValue::Init(const TArray<FCapturedPropSegment>& InCapturedPropSegm
 	TempObjPtr.Reset();
 }
 
+void UPropertyValue::BeginDestroy()
+{
+	Super::BeginDestroy();
+
+#if WITH_EDITOR
+	FEditorDelegates::EndPIE.RemoveAll(this);
+#endif
+}
+
+#if WITH_EDITOR
+
+void UPropertyValue::OnPIEEnded(const bool bIsSimulatingInEditor)
+{
+	ClearLastResolve();
+}
+
+#endif
+
 UVariantObjectBinding* UPropertyValue::GetParent() const
 {
 	return Cast<UVariantObjectBinding>(GetOuter());
diff --git a/Engine/Plugins/Enterprise/VariantManagerContent/Source/VariantManagerContent/Public/PropertyValue.h b/Engine/Plugins/Enterprise/VariantManagerContent/Source/VariantManagerContent/Public/PropertyValue.h
--- a/Engine/Plugins/Enterprise/VariantManagerContent/Source/VariantManagerContent/Public/PropertyValue.h
+++ b/Engine/Plugins/Enterprise/VariantManagerContent/Source/VariantManagerContent/Public/PropertyValue.h
@@ -73,6 +73,7 @@ public:
 
 	// UObject Interface
 	virtual void Serialize(FArchive& Ar) override;
+	virtual void BeginDestroy() override;
 	//~ End UObject Interface
 
 	// Tries to resolve the property value on the passed object, or the parent binding's bound object if the argument is nullptr
@@ -173,6 +174,10 @@ protected:
 
 	FOnPropertyApplied OnPropertyApplied;
 	FOnPropertyRecorded OnPropertyRecorded;
+	
+#if WITH_EDITOR
+	void OnPIEEnded(const bool bIsSimulatingInEditor);
+#endif
 
 	// Temp data cached from last resolve
 	FProperty* LeafProperty;
