Fixed a bunch of non unity errors

diff --git a/Engine/Plugins/Editor/AssetManagerEditor/Source/AssetManagerEditor/Private/ReferenceViewer/SReferenceViewer.h b/Engine/Plugins/Editor/AssetManagerEditor/Source/AssetManagerEditor/Private/ReferenceViewer/SReferenceViewer.h
--- a/Engine/Plugins/Editor/AssetManagerEditor/Source/AssetManagerEditor/Private/ReferenceViewer/SReferenceViewer.h
+++ b/Engine/Plugins/Editor/AssetManagerEditor/Source/AssetManagerEditor/Private/ReferenceViewer/SReferenceViewer.h
@@ -8,6 +8,7 @@
 #include "CollectionManagerTypes.h"
 
 class FUICommandList;
+class SComboButton;
 class SReferenceViewerFilterBar;
 class SSearchBox;
 class UEdGraph_ReferenceViewer;
diff --git a/Engine/Plugins/Enterprise/VariantManagerContent/Source/VariantManagerContent/Private/LevelVariantSets.cpp b/Engine/Plugins/Enterprise/VariantManagerContent/Source/VariantManagerContent/Private/LevelVariantSets.cpp
--- a/Engine/Plugins/Enterprise/VariantManagerContent/Source/VariantManagerContent/Private/LevelVariantSets.cpp
+++ b/Engine/Plugins/Enterprise/VariantManagerContent/Source/VariantManagerContent/Private/LevelVariantSets.cpp
@@ -17,6 +17,9 @@
 #include "Editor.h"
 #include "GameDelegates.h"
 #include "Engine/Blueprint.h"
+#else
+#include "Engine/World.h"
+#include "UObject/Package.h"
 #endif
 
 #define LOCTEXT_NAMESPACE "LevelVariantSets"
diff --git a/Engine/Plugins/Experimental/Mutable/Source/CustomizableObject/Private/MuCO/CustomizableObjectSystem.cpp b/Engine/Plugins/Experimental/Mutable/Source/CustomizableObject/Private/MuCO/CustomizableObjectSystem.cpp
--- a/Engine/Plugins/Experimental/Mutable/Source/CustomizableObject/Private/MuCO/CustomizableObjectSystem.cpp
+++ b/Engine/Plugins/Experimental/Mutable/Source/CustomizableObject/Private/MuCO/CustomizableObjectSystem.cpp
@@ -34,6 +34,8 @@
 #include "Logging/MessageLog.h"
 #include "Misc/ConfigCacheIni.h"
 #include "Misc/MessageDialog.h"
+#else
+#include "Engine/Engine.h"
 #endif
 
 
diff --git a/Engine/Plugins/Experimental/PCG/Source/PCG/Private/PCGGraph.cpp b/Engine/Plugins/Experimental/PCG/Source/PCG/Private/PCGGraph.cpp
--- a/Engine/Plugins/Experimental/PCG/Source/PCG/Private/PCGGraph.cpp
+++ b/Engine/Plugins/Experimental/PCG/Source/PCG/Private/PCGGraph.cpp
@@ -10,6 +10,8 @@
 
 #if WITH_EDITOR
 #include "Editor.h"
+#else
+#include "UObject/Package.h"
 #endif
 
 UPCGGraph::UPCGGraph(const FObjectInitializer& ObjectInitializer)
diff --git a/Engine/Plugins/Experimental/PCG/Source/PCG/Private/PCGSubsystem.cpp b/Engine/Plugins/Experimental/PCG/Source/PCG/Private/PCGSubsystem.cpp
--- a/Engine/Plugins/Experimental/PCG/Source/PCG/Private/PCGSubsystem.cpp
+++ b/Engine/Plugins/Experimental/PCG/Source/PCG/Private/PCGSubsystem.cpp
@@ -23,6 +23,8 @@
 #include "WorldPartition/WorldPartitionHelpers.h"
 #include "WorldPartition/ActorPartition/PartitionActorDesc.h"
 #include "ObjectTools.h"
+#else
+#include "Engine/World.h"
 #endif
 
 namespace PCGSubsystemConsole
diff --git a/Engine/Plugins/Experimental/SkeletalReduction/Source/Private/SkeletalMeshReductionPlugin.cpp b/Engine/Plugins/Experimental/SkeletalReduction/Source/Private/SkeletalMeshReductionPlugin.cpp
--- a/Engine/Plugins/Experimental/SkeletalReduction/Source/Private/SkeletalMeshReductionPlugin.cpp
+++ b/Engine/Plugins/Experimental/SkeletalReduction/Source/Private/SkeletalMeshReductionPlugin.cpp
@@ -3,6 +3,7 @@
 
 
 #include "AnimationRuntime.h"
+#include "Animation/AttributesRuntime.h"
 #include "Engine/SkeletalMesh.h"
 #include "Engine/SkinnedAssetCommon.h"
 #include "Features/IModularFeatures.h"
diff --git a/Engine/Plugins/Runtime/AppleImageUtils/Source/AppleImageUtils/Private/AppleImageUtilsPlugin.cpp b/Engine/Plugins/Runtime/AppleImageUtils/Source/AppleImageUtils/Private/AppleImageUtilsPlugin.cpp
--- a/Engine/Plugins/Runtime/AppleImageUtils/Source/AppleImageUtils/Private/AppleImageUtilsPlugin.cpp
+++ b/Engine/Plugins/Runtime/AppleImageUtils/Source/AppleImageUtils/Private/AppleImageUtilsPlugin.cpp
@@ -5,6 +5,8 @@
 
 #if SUPPORTS_IMAGE_UTILS_1_0
 	#include "Apple/ApplePlatformMisc.h"
+	#include "Async/Async.h"
+	#include "Engine/Texture2D.h"
 
 	#import <CoreImage/CIContext.h>
 
diff --git a/Engine/Plugins/Runtime/DataRegistry/Source/DataRegistry/Private/DataRegistrySubsystem.cpp b/Engine/Plugins/Runtime/DataRegistry/Source/DataRegistry/Private/DataRegistrySubsystem.cpp
--- a/Engine/Plugins/Runtime/DataRegistry/Source/DataRegistry/Private/DataRegistrySubsystem.cpp
+++ b/Engine/Plugins/Runtime/DataRegistry/Source/DataRegistry/Private/DataRegistrySubsystem.cpp
@@ -16,6 +16,9 @@
 
 #if WITH_EDITOR
 #include "Editor.h"
+#else
+#include "AssetRegistry/AssetData.h"
+#include "Engine/Engine.h"
 #endif
 
 UDataRegistrySubsystem* UDataRegistrySubsystem::Get()
diff --git a/Engine/Plugins/Runtime/MassEntity/Source/MassEntity/Public/MassSubsystemAccess.h b/Engine/Plugins/Runtime/MassEntity/Source/MassEntity/Public/MassSubsystemAccess.h
--- a/Engine/Plugins/Runtime/MassEntity/Source/MassEntity/Public/MassSubsystemAccess.h
+++ b/Engine/Plugins/Runtime/MassEntity/Source/MassEntity/Public/MassSubsystemAccess.h
@@ -14,6 +14,8 @@
 #if WITH_EDITOR
 #include "Editor.h"
 #include "EditorSubsystem.h"
+#else
+#include "Engine/Engine.h"
 #endif // WITH_EDITOR
 
 
diff --git a/Engine/Source/Runtime/NavigationSystem/Public/NavigationInvokerComponent.h b/Engine/Source/Runtime/NavigationSystem/Public/NavigationInvokerComponent.h
--- a/Engine/Source/Runtime/NavigationSystem/Public/NavigationInvokerComponent.h
+++ b/Engine/Source/Runtime/NavigationSystem/Public/NavigationInvokerComponent.h
@@ -4,6 +4,7 @@
 #include "CoreMinimal.h"
 #include "UObject/ObjectMacros.h"
 #include "UObject/UObjectGlobals.h"
+#include "AI/Navigation/NavAgentSelector.h"
 #include "AI/Navigation/NavigationTypes.h"
 #include "Components/ActorComponent.h"
 #include "NavigationInvokerComponent.generated.h"
