LevelVariantSets - Replace usage of IAssetTypeActions_AssetInterface with UAssetDefinition

diff --git a/Engine/Plugins/Enterprise/VariantManagerContent/Source/VariantManagerContentEditor/Private/AssetDefinition_LevelVariantSets.cpp b/Engine/Plugins/Enterprise/VariantManagerContent/Source/VariantManagerContentEditor/Private/AssetDefinition_LevelVariantSets.cpp
new file mode 100644
--- /dev/null
+++ b/Engine/Plugins/Enterprise/VariantManagerContent/Source/VariantManagerContentEditor/Private/AssetDefinition_LevelVariantSets.cpp
@@ -0,0 +1,87 @@
+// Copyright Epic Games, Inc. All Rights Reserved.
+
+#include "AssetDefinition_LevelVariantSets.h"
+
+#include "ContentBrowserMenuContexts.h"
+#include "LevelVariantSets.h"
+#include "ToolMenuSection.h"
+#include "VariantManagerContentEditorModule.h"
+
+
+#define LOCTEXT_NAMESPACE "LevelVariantSetAssetActions"
+
+FText UAssetDefinition_LevelVariantSets::GetAssetDisplayName() const
+{
+	return NSLOCTEXT("AssetTypeActions", "AssetTypeActions_LevelVariantSets", "Level Variant Sets");
+}
+
+TSoftClassPtr<> UAssetDefinition_LevelVariantSets::GetAssetClass() const
+{
+	return ULevelVariantSets::StaticClass();
+}
+
+EAssetCommandResult UAssetDefinition_LevelVariantSets::OpenAssets(const FAssetOpenArgs& OpenArgs) const
+{
+	EToolkitMode::Type Mode = OpenArgs.ToolkitHost.IsValid()
+		? EToolkitMode::WorldCentric
+		: EToolkitMode::Standalone;
+
+	for (ULevelVariantSets* LevelVariantSets : OpenArgs.LoadObjects<ULevelVariantSets>())
+	{
+		IVariantManagerContentEditorModule& ContentEditorModule = IVariantManagerContentEditorModule::Get();
+		ContentEditorModule.GetOnLevelVariantSetsEditorOpened().ExecuteIfBound(Mode, OpenArgs.ToolkitHost, LevelVariantSets);
+	}
+
+	return EAssetCommandResult::Handled;
+}
+
+TConstArrayView<FAssetCategoryPath> UAssetDefinition_LevelVariantSets::GetAssetCategories() const
+{
+	static const auto Categories = { EAssetCategoryPaths::Misc };
+	return Categories;
+}
+
+FAssetSupportResponse UAssetDefinition_LevelVariantSets::CanLocalize(const FAssetData& InAsset) const
+{
+	return FAssetSupportResponse::NotSupported();
+}
+
+
+namespace MenuExtension_LevelVariantSets
+{
+	static FDelayedAutoRegisterHelper DelayedAutoRegister(EDelayedRegisterRunPhase::EndOfEngineInit, []{ 
+		UToolMenus::RegisterStartupCallback(FSimpleMulticastDelegate::FDelegate::CreateLambda([]()
+		{
+			FToolMenuOwnerScoped OwnerScoped(UE_MODULE_NAME);
+			UToolMenu* Menu = UE::ContentBrowser::ExtendToolMenu_AssetContextMenu(ULevelVariantSets::StaticClass());
+	        
+			FToolMenuSection& Section = Menu->FindOrAddSection("GetAssetActions");
+			Section.AddDynamicEntry(NAME_None, FNewToolMenuSectionDelegate::CreateLambda([](FToolMenuSection& InSection)
+			{
+				if (const UContentBrowserAssetContextMenuContext* CBContext = UContentBrowserAssetContextMenuContext::FindContextWithAssets(InSection))
+				{
+					TArray<FAssetData> Assets = CBContext->GetSelectedAssetsOfType(ULevelVariantSets::StaticClass());
+
+					InSection.AddMenuEntry(
+						"LevelVariantSetAssetActions_CreateActorText",
+						LOCTEXT("CreateActorText", "Create LevelVariantSets actor"),
+						LOCTEXT("CreateActorTooltip", "Creates a new ALevelVariantSetsActor AActor and add it to the scene"),
+						FSlateIcon(),
+						FToolMenuExecuteAction::CreateLambda([](const FToolMenuContext& Context)
+						{
+							if (const UContentBrowserAssetContextMenuContext* CBContext = UContentBrowserAssetContextMenuContext::FindContextWithAssets(Context))
+							{
+								IVariantManagerContentEditorModule& ContentEditorModule = IVariantManagerContentEditorModule::Get();
+								for (ULevelVariantSets* LevelVariantSets : CBContext->LoadSelectedObjects<ULevelVariantSets>())
+								{
+									ContentEditorModule.GetOrCreateLevelVariantSetsActor(LevelVariantSets, true);
+								}
+							}
+						}));
+				}
+			}));
+		}));
+	});
+}
+
+#undef LOCTEXT_NAMESPACE
diff --git a/Engine/Plugins/Enterprise/VariantManagerContent/Source/VariantManagerContentEditor/Private/AssetDefinition_LevelVariantSets.h b/Engine/Plugins/Enterprise/VariantManagerContent/Source/VariantManagerContentEditor/Private/AssetDefinition_LevelVariantSets.h
new file mode 100644
--- /dev/null
+++ b/Engine/Plugins/Enterprise/VariantManagerContent/Source/VariantManagerContentEditor/Private/AssetDefinition_LevelVariantSets.h
@@ -0,0 +1,26 @@
+// Copyright Epic Games, Inc. All Rights Reserved.
+
+#pragma once
+
+#include "AssetDefinitionDefault.h"
+#include "AssetTypeActions_Base.h"
+
+#include "AssetDefinition_LevelVariantSets.generated.h"
+
+
+UCLASS()
+class UAssetDefinition_LevelVariantSets : public UAssetDefinitionDefault
+{
+	GENERATED_BODY()
+
+public:
+	// UAssetDefinition Begin
+	virtual FText GetAssetDisplayName() const override;
+	virtual FLinearColor GetAssetColor() const override { return FColor(80, 80, 200); }
+	virtual TSoftClassPtr<UObject> GetAssetClass() const override;
+	virtual EAssetCommandResult OpenAssets(const FAssetOpenArgs& OpenArgs) const override;
+	virtual TConstArrayView<FAssetCategoryPath> GetAssetCategories() const override;
+	virtual bool CanImport() const override { return false; }
+	virtual FAssetSupportResponse CanLocalize(const FAssetData& InAsset) const override;
+	// UAssetDefinition End
+};
diff --git a/Engine/Plugins/Enterprise/VariantManagerContent/Source/VariantManagerContentEditor/Private/LevelVariantSetsAssetActions.cpp b/Engine/Plugins/Enterprise/VariantManagerContent/Source/VariantManagerContentEditor/Private/LevelVariantSetsAssetActions.cpp
deleted file mode 100644
diff --git a/Engine/Plugins/Enterprise/VariantManagerContent/Source/VariantManagerContentEditor/Private/LevelVariantSetsAssetActions.h b/Engine/Plugins/Enterprise/VariantManagerContent/Source/VariantManagerContentEditor/Private/LevelVariantSetsAssetActions.h
deleted file mode 100644
diff --git a/Engine/Plugins/Enterprise/VariantManagerContent/Source/VariantManagerContentEditor/Private/VariantManagerContentEditorModule.cpp b/Engine/Plugins/Enterprise/VariantManagerContent/Source/VariantManagerContentEditor/Private/VariantManagerContentEditorModule.cpp
--- a/Engine/Plugins/Enterprise/VariantManagerContent/Source/VariantManagerContentEditor/Private/VariantManagerContentEditorModule.cpp
+++ b/Engine/Plugins/Enterprise/VariantManagerContent/Source/VariantManagerContentEditor/Private/VariantManagerContentEditorModule.cpp
@@ -2,14 +2,15 @@
 
 #include "VariantManagerContentEditorModule.h"
 
+#include "Editor.h"
 #include "LevelVariantSets.h"
 #include "LevelVariantSetsActor.h"
 #include "LevelVariantSetsActorCustomization.h"
-#include "LevelVariantSetsAssetActions.h"
 #include "SwitchActorCustomization.h"
 #include "UObject/UObjectIterator.h"
 #include "VariantManagerContentEditorLog.h"
 
+#include "AssetToolsModule.h"
 #include "Factories/Factory.h"
 #include "ObjectTools.h"
 #include "PackageTools.h"
@@ -24,11 +25,6 @@ class FVariantManagerContentEditorModule : public IVariantManagerContentEditorMo
 public:
 	virtual void StartupModule() override
 	{
-		// Register asset actions for the LevelVariantSets asset
-		IAssetTools& AssetTools = FModuleManager::LoadModuleChecked<FAssetToolsModule>(TEXT("AssetTools")).Get();
-		AssetActions = MakeShareable(new FLevelVariantSetsAssetActions());
-		AssetTools.RegisterAssetTypeActions(AssetActions.ToSharedRef());
-
 		FPropertyEditorModule& PropertyModule = FModuleManager::LoadModuleChecked<FPropertyEditorModule>(TEXT("PropertyEditor"));
 		PropertyModule.RegisterCustomClassLayout(TEXT("LevelVariantSetsActor"), FOnGetDetailCustomizationInstance::CreateStatic(&FLevelVariantSetsActorCustomization::MakeInstance));
 		PropertyModule.RegisterCustomClassLayout(TEXT("SwitchActor"), FOnGetDetailCustomizationInstance::CreateStatic(&FSwitchActorCustomization::MakeInstance));
@@ -36,14 +32,6 @@ public:
 
 	virtual void ShutdownModule() override
 	{
-        // Unregister asset actions for the LevelVariantSets asset
-		FAssetToolsModule* AssetToolsModule = FModuleManager::GetModulePtr<FAssetToolsModule>(TEXT("AssetTools"));
-		if (AssetToolsModule != nullptr)
-		{
-			IAssetTools& AssetTools = AssetToolsModule->Get();
-			AssetTools.UnregisterAssetTypeActions(AssetActions.ToSharedRef());
-		}
-
         FPropertyEditorModule& PropertyModule = FModuleManager::LoadModuleChecked< FPropertyEditorModule >(TEXT("PropertyEditor"));
 		PropertyModule.UnregisterCustomClassLayout(TEXT("LevelVariantSetsActor"));
 		PropertyModule.UnregisterCustomClassLayout(TEXT("SwitchActor"));
@@ -191,7 +179,6 @@ public:
 
 private:
 	FOnLevelVariantSetsEditor OnLevelVariantSetsEditor;
-	TSharedPtr<IAssetTypeActions> AssetActions;
 };
 
 IMPLEMENT_MODULE(FVariantManagerContentEditorModule, VariantManagerContentEditor);
diff --git a/Engine/Plugins/Enterprise/VariantManagerContent/Source/VariantManagerContentEditor/VariantManagerContentEditor.Build.cs b/Engine/Plugins/Enterprise/VariantManagerContent/Source/VariantManagerContentEditor/VariantManagerContentEditor.Build.cs
--- a/Engine/Plugins/Enterprise/VariantManagerContent/Source/VariantManagerContentEditor/VariantManagerContentEditor.Build.cs
+++ b/Engine/Plugins/Enterprise/VariantManagerContent/Source/VariantManagerContentEditor/VariantManagerContentEditor.Build.cs
@@ -10,10 +10,13 @@ namespace UnrealBuildTool.Rules
 			PublicDependencyModuleNames.AddRange(
 				new string[]
 				{
+					"AssetDefinition",
+					"ContentBrowser",
 					"Core",
 					"CoreUObject",
 					"EditorFramework",
 					"Engine",
+					"ToolMenus",
 					"UnrealEd",
 				}
 			);
