Don't operate on UObjects during object destruction, prefer doing so inside BeginDestroy instead.

#rnx

diff --git a/Engine/Plugins/Enterprise/VariantManagerContent/Source/VariantManagerContent/Private/LevelVariantSets.cpp b/Engine/Plugins/Enterprise/VariantManagerContent/Source/VariantManagerContent/Private/LevelVariantSets.cpp
--- a/Engine/Plugins/Enterprise/VariantManagerContent/Source/VariantManagerContent/Private/LevelVariantSets.cpp
+++ b/Engine/Plugins/Enterprise/VariantManagerContent/Source/VariantManagerContent/Private/LevelVariantSets.cpp
@@ -78,8 +78,10 @@ ULevelVariantSets::ULevelVariantSets(const FObjectInitializer& ObjectInitializer
 #endif
 }
 
-ULevelVariantSets::~ULevelVariantSets()
+
+void ULevelVariantSets::BeginDestroy()
 {
+	Super::BeginDestroy();
 #if WITH_EDITOR
 	UnsubscribeToEditorDelegates();
 	UnsubscribeToDirectorCompiled();
diff --git a/Engine/Plugins/Enterprise/VariantManagerContent/Source/VariantManagerContent/Public/LevelVariantSets.h b/Engine/Plugins/Enterprise/VariantManagerContent/Source/VariantManagerContent/Public/LevelVariantSets.h
--- a/Engine/Plugins/Enterprise/VariantManagerContent/Source/VariantManagerContent/Public/LevelVariantSets.h
+++ b/Engine/Plugins/Enterprise/VariantManagerContent/Source/VariantManagerContent/Public/LevelVariantSets.h
@@ -18,9 +18,8 @@ class VARIANTMANAGERCONTENT_API ULevelVariantSets : public UObject
 
 public:
 
-	~ULevelVariantSets();
-
 	// UObject interface
+	virtual void BeginDestroy() override;
 	virtual void Serialize(FArchive& Ar) override;
 	//~ End UObject interface
 
