Fixed VariantManager opens an extra empty window

diff --git a/Engine/Plugins/Enterprise/VariantManagerContent/Source/VariantManagerContentEditor/Private/AssetDefinition_LevelVariantSets.cpp b/Engine/Plugins/Enterprise/VariantManagerContent/Source/VariantManagerContentEditor/Private/AssetDefinition_LevelVariantSets.cpp
--- a/Engine/Plugins/Enterprise/VariantManagerContent/Source/VariantManagerContentEditor/Private/AssetDefinition_LevelVariantSets.cpp
+++ b/Engine/Plugins/Enterprise/VariantManagerContent/Source/VariantManagerContentEditor/Private/AssetDefinition_LevelVariantSets.cpp
@@ -22,14 +22,10 @@ TSoftClassPtr<> UAssetDefinition_LevelVariantSets::GetAssetClass() const
 
 EAssetCommandResult UAssetDefinition_LevelVariantSets::OpenAssets(const FAssetOpenArgs& OpenArgs) const
 {
-	EToolkitMode::Type Mode = OpenArgs.ToolkitHost.IsValid()
-		? EToolkitMode::WorldCentric
-		: EToolkitMode::Standalone;
-
 	for (ULevelVariantSets* LevelVariantSets : OpenArgs.LoadObjects<ULevelVariantSets>())
 	{
 		IVariantManagerContentEditorModule& ContentEditorModule = IVariantManagerContentEditorModule::Get();
-		ContentEditorModule.GetOnLevelVariantSetsEditorOpened().ExecuteIfBound(Mode, OpenArgs.ToolkitHost, LevelVariantSets);
+		ContentEditorModule.GetOnLevelVariantSetsEditorOpened().ExecuteIfBound(OpenArgs.GetToolkitMode(), OpenArgs.ToolkitHost, LevelVariantSets);
 	}
 
 	return EAssetCommandResult::Handled;
@@ -46,6 +42,10 @@ FAssetSupportResponse UAssetDefinition_LevelVariantSets::CanLocalize(const FAsse
 	return FAssetSupportResponse::NotSupported();
 }
 
+FAssetOpenSupport UAssetDefinition_LevelVariantSets::GetAssetOpenSupport(const FAssetOpenSupportArgs& OpenSupportArgs) const
+{
+	return FAssetOpenSupport(OpenSupportArgs.OpenMethod,OpenSupportArgs.OpenMethod == EAssetOpenMethod::Edit, EToolkitMode::WorldCentric); 
+}
 
 namespace MenuExtension_LevelVariantSets
 {
diff --git a/Engine/Plugins/Enterprise/VariantManagerContent/Source/VariantManagerContentEditor/Private/AssetDefinition_LevelVariantSets.h b/Engine/Plugins/Enterprise/VariantManagerContent/Source/VariantManagerContentEditor/Private/AssetDefinition_LevelVariantSets.h
--- a/Engine/Plugins/Enterprise/VariantManagerContent/Source/VariantManagerContentEditor/Private/AssetDefinition_LevelVariantSets.h
+++ b/Engine/Plugins/Enterprise/VariantManagerContent/Source/VariantManagerContentEditor/Private/AssetDefinition_LevelVariantSets.h
@@ -22,5 +22,6 @@ public:
 	virtual TConstArrayView<FAssetCategoryPath> GetAssetCategories() const override;
 	virtual bool CanImport() const override { return false; }
 	virtual FAssetSupportResponse CanLocalize(const FAssetData& InAsset) const override;
+	virtual FAssetOpenSupport GetAssetOpenSupport(const FAssetOpenSupportArgs& OpenSupportArgs) const override;
 	// UAssetDefinition End
 };
