Added defensive code against  potential null pointers

#rnx

diff --git a/Engine/Plugins/Enterprise/VariantManagerContent/Source/VariantManagerContent/Private/VariantObjectBinding.cpp b/Engine/Plugins/Enterprise/VariantManagerContent/Source/VariantManagerContent/Private/VariantObjectBinding.cpp
--- a/Engine/Plugins/Enterprise/VariantManagerContent/Source/VariantManagerContent/Private/VariantObjectBinding.cpp
+++ b/Engine/Plugins/Enterprise/VariantManagerContent/Source/VariantManagerContent/Private/VariantObjectBinding.cpp
@@ -381,7 +381,11 @@ void UVariantObjectBinding::ExecuteAllTargetFunctions()
 void UVariantObjectBinding::UpdateFunctionCallerNames()
 {
 	ULevelVariantSets* ParentLVS = GetTypedOuter<ULevelVariantSets>();
-	UObject* DirectorInstance = ParentLVS->GetDirectorInstance(GetObject());
+	UObject* DirectorInstance = ParentLVS ? ParentLVS->GetDirectorInstance(GetObject()) : nullptr;
+	if (!DirectorInstance)
+	{
+		return;
+	}
 
 	bool bHasChanged = false;
 
