Speculative fix for invalid TWeakPtr

#rnx

diff --git a/Engine/Plugins/Enterprise/VariantManagerContent/Source/VariantManagerContent/Private/LevelVariantSets.cpp b/Engine/Plugins/Enterprise/VariantManagerContent/Source/VariantManagerContent/Private/LevelVariantSets.cpp
--- a/Engine/Plugins/Enterprise/VariantManagerContent/Source/VariantManagerContent/Private/LevelVariantSets.cpp
+++ b/Engine/Plugins/Enterprise/VariantManagerContent/Source/VariantManagerContent/Private/LevelVariantSets.cpp
@@ -497,9 +497,14 @@ void ULevelVariantSets::UnsubscribeToDirectorCompiled()
 
 void ULevelVariantSets::HandleDirectorDestroyed(ULevelVariantSetsFunctionDirector* Director)
 {
+	if (!IsValid(Director))
+	{
+		return;
+	}
+
 	for (TMap<UWorld*, TWeakObjectPtr<UObject>>::TIterator Iter(WorldToDirectorInstance); Iter; ++Iter )
 	{
-		if ( Iter->Value == Director)
+		if (Iter->Value.IsValid() && Iter->Value == Director)
 		{
 			Iter.RemoveCurrent();
 		}
