Additional fix for crash when deleting ULevelVariantSetsFunctionDirector object

#rnx

diff --git a/Engine/Plugins/Enterprise/VariantManagerContent/Source/VariantManagerContent/Public/LevelVariantSetsFunctionDirector.h b/Engine/Plugins/Enterprise/VariantManagerContent/Source/VariantManagerContent/Public/LevelVariantSetsFunctionDirector.h
--- a/Engine/Plugins/Enterprise/VariantManagerContent/Source/VariantManagerContent/Public/LevelVariantSetsFunctionDirector.h
+++ b/Engine/Plugins/Enterprise/VariantManagerContent/Source/VariantManagerContent/Public/LevelVariantSetsFunctionDirector.h
@@ -15,8 +15,13 @@ public:
 	GENERATED_BODY()
 
 	~ULevelVariantSetsFunctionDirector()
+	{
+	}
+
+	virtual void BeginDestroy() override
 	{
 		OnDestroy.Broadcast(this);
+		Super::BeginDestroy();
 	}
 
 	OnDirectorDestroyed& GetOnDestroy()
