TLazyObjectPtr Deprecation:
* VariantManagerContent plugin: pragma wrap TLazyObjectPtr fallback usage

#Synced-CL 52628728

diff --git a/Engine/Plugins/Enterprise/VariantManagerContent/Source/VariantManagerContent/Private/VariantObjectBinding.cpp b/Engine/Plugins/Enterprise/VariantManagerContent/Source/VariantManagerContent/Private/VariantObjectBinding.cpp
--- a/Engine/Plugins/Enterprise/VariantManagerContent/Source/VariantManagerContent/Private/VariantObjectBinding.cpp
+++ b/Engine/Plugins/Enterprise/VariantManagerContent/Source/VariantManagerContent/Private/VariantObjectBinding.cpp
@@ -31,7 +31,9 @@ void UVariantObjectBinding::SetObject(UObject* InObject)
 	Modify();
 
 	ObjectPtr = InObject;
-	LazyObjectPtr = InObject;
+	PRAGMA_DISABLE_DEPRECATION_WARNINGS
+	LazyObjectPtr_DEPRECATED = InObject;
+	PRAGMA_ENABLE_DEPRECATION_WARNINGS
 }
 
 UVariant* UVariantObjectBinding::GetParent()
@@ -133,19 +135,25 @@ UObject* UVariantObjectBinding::GetObject() const
 		UObject* Obj = TempPtr.ResolveObject();
 		if (IsValid(Obj) && !Obj->IsUnreachable())
 		{
-			LazyObjectPtr = Obj;
+			PRAGMA_DISABLE_DEPRECATION_WARNINGS
+			LazyObjectPtr_DEPRECATED = Obj;
+			PRAGMA_ENABLE_DEPRECATION_WARNINGS
 			return Obj;
 		}
 		// Fixup for redirectors (e.g. when going from temp level to a saved level)
 		// I could do ObjectPtr.PreSavePath, which in fact follows redirectors. This doesn't work
 		// for saving after moving to a new level and then reloading, as the redirector will
-		// only be created AFTER we saved. The LazyObjectPtr successfully manages to track
+		// only be created AFTER we saved. The LazyObjectPtr_DEPRECATED successfully manages to track
 		// the object across levels, however. We don't exclusively use this because it is not meant
 		// to update to the duplicated objects when going into PIE
 		// This could potentially be enclosed in a #if WITH_EDITOR block
+		// Keeping this deprecated part as a stop gap for that edge case until we remove the LazyObjectPtr code
+		// At that time a FUniversalObjectLocator style approach might be the alternative
 		else
 		{
-			UObject* LazyObject = LazyObjectPtr.Get();
+			PRAGMA_DISABLE_DEPRECATION_WARNINGS
+			UObject* LazyObject = LazyObjectPtr_DEPRECATED.Get();
+			PRAGMA_ENABLE_DEPRECATION_WARNINGS
 			if (LazyObject)
 			{
 				//UE_LOGF(LogVariantContent, Log, "Actor '%ls' switched path. Binding updating from path '%ls' to '%ls'", *LazyObject->GetName(), *ObjectPtr.ToString(), *LazyObject->GetFullName());
diff --git a/Engine/Plugins/Enterprise/VariantManagerContent/Source/VariantManagerContent/Public/VariantObjectBinding.h b/Engine/Plugins/Enterprise/VariantManagerContent/Source/VariantManagerContent/Public/VariantObjectBinding.h
--- a/Engine/Plugins/Enterprise/VariantManagerContent/Source/VariantManagerContent/Public/VariantObjectBinding.h
+++ b/Engine/Plugins/Enterprise/VariantManagerContent/Source/VariantManagerContent/Public/VariantObjectBinding.h
@@ -58,8 +58,10 @@ private:
 	UPROPERTY()
 	mutable FSoftObjectPath ObjectPtr;
 
+	PRAGMA_DISABLE_DEPRECATION_WARNINGS
 	UPROPERTY()
-	mutable TLazyObjectPtr<UObject> LazyObjectPtr;
+	mutable TLazyObjectPtr<UObject> LazyObjectPtr_DEPRECATED;
+	PRAGMA_ENABLE_DEPRECATION_WARNINGS
 
 	UPROPERTY()
 	TArray<TObjectPtr<UPropertyValue>> CapturedProperties;
